Move AudioProperties to public experimental
#Synced-CL 49403902

diff --git a/Engine/Plugins/Experimental/AudioProperties/AudioProperties.uplugin b/Engine/Plugins/Experimental/AudioProperties/AudioProperties.uplugin
new file mode 100644
--- /dev/null
+++ b/Engine/Plugins/Experimental/AudioProperties/AudioProperties.uplugin
@@ -0,0 +1,32 @@
+{
+	"FileVersion": 3,
+	"Version": 0,
+	"VersionName": "0.0",
+	"FriendlyName": "Audio Properties",
+	"Description": "Allows to define arbitrary derivable sets of properties to be injected in any audio asset",
+	"Category": "Audio",
+	"CreatedBy": "Epic Games, Inc.",
+	"CreatedByURL": "http://epicgames.com",
+	"DocsURL": "",
+	"MarketplaceURL": "",
+	"SupportURL": "",
+	"EnabledByDefault": false,
+	"CanContainContent": true,
+	"IsBetaVersion": false,
+	"IsExperimentalVersion": true,
+	"Installed": false,
+	"Modules": [
+		{
+			"Name": "AudioProperties",
+			"Type": "Editor",
+			"LoadingPhase": "Default"
+		},
+		{
+			"Name": "AudioPropertiesEditor",
+			"Type": "Editor",
+			"LoadingPhase": "Default"
+		}
+	],
+	"Plugins": [
+	]
+}
\ No newline at end of file
diff --git a/Engine/Plugins/Experimental/AudioProperties/Source/AudioProperties/AudioProperties.build.cs b/Engine/Plugins/Experimental/AudioProperties/Source/AudioProperties/AudioProperties.build.cs
new file mode 100644
--- /dev/null
+++ b/Engine/Plugins/Experimental/AudioProperties/Source/AudioProperties/AudioProperties.build.cs
@@ -0,0 +1,25 @@
+// Copyright Epic Games, Inc. All Rights Reserved.
+
+using UnrealBuildTool;
+using System.IO;
+
+namespace UnrealBuildTool.Rules
+{
+	public class AudioProperties : ModuleRules
+	{
+		public AudioProperties(ReadOnlyTargetRules Target) : base(Target)
+		{
+			PCHUsage = PCHUsageMode.UseExplicitOrSharedPCHs;
+
+			PublicDependencyModuleNames.AddRange(
+				new string[] {
+					"AudioExtensions",
+					"AssetRegistry",
+					"Core",
+					"CoreUObject",
+					"TargetPlatform"
+				}
+			);
+		}
+	}
+}
\ No newline at end of file
diff --git a/Engine/Plugins/Experimental/AudioProperties/Source/AudioProperties/Private/AudioPropertiesLogs.cpp b/Engine/Plugins/Experimental/AudioProperties/Source/AudioProperties/Private/AudioPropertiesLogs.cpp
new file mode 100644
--- /dev/null
+++ b/Engine/Plugins/Experimental/AudioProperties/Source/AudioProperties/Private/AudioPropertiesLogs.cpp
@@ -0,0 +1,5 @@
+// Copyright Epic Games, Inc. All Rights Reserved.
+
+#include "AudioPropertiesLogs.h"
+
+DEFINE_LOG_CATEGORY(LogAudioProperties);
diff --git a/Engine/Plugins/Experimental/AudioProperties/Source/AudioProperties/Private/AudioPropertiesModule.cpp b/Engine/Plugins/Experimental/AudioProperties/Source/AudioProperties/Private/AudioPropertiesModule.cpp
new file mode 100644
--- /dev/null
+++ b/Engine/Plugins/Experimental/AudioProperties/Source/AudioProperties/Private/AudioPropertiesModule.cpp
@@ -0,0 +1,14 @@
+// Copyright Epic Games, Inc. All Rights Reserved.
+
+#include "AudioPropertiesModule.h"
+#include "Modules/ModuleManager.h"
+
+void FAudioPropertiesModule::StartupModule()
+{
+}
+
+void FAudioPropertiesModule::ShutdownModule()
+{
+}
+
+IMPLEMENT_MODULE(FAudioPropertiesModule, AudioProperties);
\ No newline at end of file
diff --git a/Engine/Plugins/Experimental/AudioProperties/Source/AudioProperties/Private/AudioPropertiesParser.cpp b/Engine/Plugins/Experimental/AudioProperties/Source/AudioProperties/Private/AudioPropertiesParser.cpp
new file mode 100644
--- /dev/null
+++ b/Engine/Plugins/Experimental/AudioProperties/Source/AudioProperties/Private/AudioPropertiesParser.cpp
@@ -0,0 +1,641 @@
+// Copyright Epic Games, Inc. All Rights Reserved.
+
+#include "AudioPropertiesParser.h"
+
+#include "AssetRegistry/AssetData.h"
+#include "AudioPropertiesBindings.h"
+#include "AudioPropertiesLogs.h"
+#include "AudioPropertiesParsingData.h"
+#include "AudioPropertiesSheet.h"
+#include "AudioPropertiesUtils.h"
+#include "UObject/PropertyAccessUtil.h"
+#include "UObject/UnrealTypePrivate.h"
+
+namespace AudioPropertiesParserPrivate
+{
+	static const FName ClampMinTagName = "ClampMin";
+	static const FName ClampMaxTagName = "ClampMax";
+	static const FName UIMinTagName = "UIMin";
+	static const FName UIMaxTagName = "UIMax";
+	static const FName ValidEnumValuesTagName = "ValidEnumValues";
+	static const FName InvalidEnumValuesTagName = "InvalidEnumValues";
+
+	bool IsPropertyLocallyIgnored(UObject& TargetObject, FProperty& TargetProperty)
+	{
+#if WITH_EDITORONLY_DATA
+		const IAudioPropertiesSheetAssetUserInterface* AsPropertySheetUser = Cast<IAudioPropertiesSheetAssetUserInterface>(&TargetObject);
+
+		if (AsPropertySheetUser)
+		{
+			return !AsPropertySheetUser->ShouldParseProperty(TargetProperty);	
+		}
+#endif
+		return false;
+	}
+
+	template< typename NumericType>
+	NumericType ClampByPropertyMetadata(const FProperty& TargetProperty, NumericType InValueToClamp)
+	{
+		static_assert(TIsArithmetic<NumericType>::Value, "Property to clamp must have numeric type");
+
+		const FString& MetaClampMinString = TargetProperty.GetMetaData(ClampMinTagName);
+		const FString& MetaClampMaxString = TargetProperty.GetMetaData(ClampMaxTagName);
+		const FString& MetaUIMinString = TargetProperty.GetMetaData(UIMinTagName);
+		const FString& MetaUIMaxString = TargetProperty.GetMetaData(UIMaxTagName);
+
+ 		const bool bShouldClamp = !MetaClampMinString.IsEmpty() || !MetaClampMaxString.IsEmpty() || !MetaUIMinString.IsEmpty() || !MetaUIMaxString.IsEmpty();
+ 
+ 		if (!bShouldClamp)
+ 		{
+ 			return MoveTemp(InValueToClamp);
+ 		}
+		
+		NumericType ClampMin = TNumericLimits<NumericType>::Lowest();
+		if (!MetaClampMinString.IsEmpty())
+		{
+			TTypeFromString<NumericType>::FromString(ClampMin, *MetaClampMinString);
+		}
+
+		NumericType ClampMax = TNumericLimits<NumericType>::Max();
+		if (!MetaClampMaxString.IsEmpty())
+		{
+			TTypeFromString<NumericType>::FromString(ClampMax, *MetaClampMaxString);
+		}
+		
+		NumericType UIMin = TNumericLimits<NumericType>::Lowest();
+		if (!MetaUIMinString.IsEmpty())
+		{
+			TTypeFromString<NumericType>::FromString(UIMin, *MetaUIMinString);
+		}
+		
+		NumericType UIMax = TNumericLimits<NumericType>::Max();
+		if (!MetaUIMaxString.IsEmpty())
+		{
+			TTypeFromString<NumericType>::FromString(UIMax, *MetaUIMaxString);
+		}
+
+		const NumericType MinAllowedValue = FMath::Max(UIMin, ClampMin);
+		const NumericType MaxAllowedValue = FMath::Min(UIMax, ClampMax); 
+		const NumericType ClampedValue = FMath::Clamp<NumericType>(InValueToClamp, MinAllowedValue, MaxAllowedValue);
+		
+		return ClampedValue;
+	}
+
+	bool ShouldApplyTargetPropertyMetadata(const FProperty& TargetProperty, const FPropertyBagPropertyDesc& SourcePropertyDesc, const FInstancedPropertyBag& PropertiesToInject)
+	{
+		bool bShouldApplyMetadata = false;
+
+		if (!SourcePropertyDesc.ContainerTypes.IsEmpty())
+		{
+			return bShouldApplyMetadata;
+		}
+		
+		const FName TargetPropertyName = SourcePropertyDesc.CachedProperty->GetFName();
+
+		switch (SourcePropertyDesc.ValueType)
+		{
+		case EPropertyBagPropertyType::Float:
+		{
+			const float SourceValue = PropertiesToInject.GetValueFloat(TargetPropertyName).GetValue();
+			const float ClampedFloatValue = AudioPropertiesParserPrivate::ClampByPropertyMetadata(TargetProperty, SourceValue);
+			bShouldApplyMetadata = SourceValue != ClampedFloatValue;
+		}
+		break;
+		case EPropertyBagPropertyType::Double:
+		{
+			const float SourceValue = PropertiesToInject.GetValueDouble(TargetPropertyName).GetValue();
+			const double ClampedDoubleValue = AudioPropertiesParserPrivate::ClampByPropertyMetadata(TargetProperty, SourceValue);
+			bShouldApplyMetadata = SourceValue != ClampedDoubleValue;
+		}
+		break;
+		case EPropertyBagPropertyType::Int32:
+		{
+			const float SourceValue = PropertiesToInject.GetValueInt32(TargetPropertyName).GetValue();
+			const int32 ClampedInt32Value = AudioPropertiesParserPrivate::ClampByPropertyMetadata(TargetProperty, SourceValue);
+			bShouldApplyMetadata = SourceValue != ClampedInt32Value;
+		}
+		break;
+		case EPropertyBagPropertyType::Int64:
+		{
+			const float SourceValue = PropertiesToInject.GetValueInt64(TargetPropertyName).GetValue();
+			const int64 ClampedInt64Value = AudioPropertiesParserPrivate::ClampByPropertyMetadata(TargetProperty, PropertiesToInject.GetValueInt64(TargetPropertyName).GetValue());
+			bShouldApplyMetadata = SourceValue != ClampedInt64Value;
+		}
+		break;
+		default:
+			break;
+		}
+
+		if (bShouldApplyMetadata)
+		{
+			UE_LOG(LogAudioProperties, Log, TEXT("Sheet Property '%s' needs to be clamped to be properly applied to Target Property '%s'"), *TargetPropertyName.ToString(), *TargetProperty.GetFName().ToString())	
+		}
+
+		return bShouldApplyMetadata;
+	}
+
+	bool ApplyTargetPropertyMetadata(const FProperty& TargetProperty, const FPropertyBagPropertyDesc& SourcePropertyDesc, FInstancedPropertyBag& PropertiesToInject)
+	{
+		bool bShouldApplyMetadata = false;
+
+		if (!SourcePropertyDesc.ContainerTypes.IsEmpty())
+		{
+			return bShouldApplyMetadata;
+		}
+		
+		const FName TargetPropertyName = SourcePropertyDesc.CachedProperty->GetFName();
+
+		switch (SourcePropertyDesc.ValueType)
+		{
+		case EPropertyBagPropertyType::Float:
+		{
+			const float SourceValue = PropertiesToInject.GetValueFloat(TargetPropertyName).GetValue();
+			const float ClampedFloatValue = AudioPropertiesParserPrivate::ClampByPropertyMetadata(TargetProperty, SourceValue);
+			bShouldApplyMetadata = SourceValue != ClampedFloatValue;
+			PropertiesToInject.SetValueFloat(TargetPropertyName, ClampedFloatValue);
+		}
+		break;
+		case EPropertyBagPropertyType::Double:
+		{
+			const float SourceValue = PropertiesToInject.GetValueDouble(TargetPropertyName).GetValue();
+			const double ClampedDoubleValue = AudioPropertiesParserPrivate::ClampByPropertyMetadata(TargetProperty, SourceValue);
+			bShouldApplyMetadata = SourceValue != ClampedDoubleValue;
+			PropertiesToInject.SetValueDouble(TargetPropertyName, ClampedDoubleValue);
+		}
+		break;
+		case EPropertyBagPropertyType::Int32:
+		{
+			const float SourceValue = PropertiesToInject.GetValueInt32(TargetPropertyName).GetValue();
+			const int32 ClampedInt32Value = AudioPropertiesParserPrivate::ClampByPropertyMetadata(TargetProperty, SourceValue);
+			bShouldApplyMetadata = SourceValue != ClampedInt32Value;
+			PropertiesToInject.SetValueInt32(TargetPropertyName, ClampedInt32Value);
+		}
+		break;
+		case EPropertyBagPropertyType::Int64:
+		{
+			const float SourceValue = PropertiesToInject.GetValueInt64(TargetPropertyName).GetValue();
+			const int64 ClampedInt64Value = AudioPropertiesParserPrivate::ClampByPropertyMetadata(TargetProperty, PropertiesToInject.GetValueInt64(TargetPropertyName).GetValue());
+			bShouldApplyMetadata = SourceValue != ClampedInt64Value;
+			PropertiesToInject.SetValueInt64(TargetPropertyName, ClampedInt64Value);
+		}
+		break;
+		default:
+			break;
+		}
+
+		return bShouldApplyMetadata;
+	}
+
+	const bool IsValidEnumForTargetMetadata(const FProperty& TargetProperty, const FPropertyBagPropertyDesc& SourcePropertyDesc, const FInstancedPropertyBag& PropertyToInjectParentBag)
+	{
+		if (SourcePropertyDesc.ValueType == EPropertyBagPropertyType::Enum && TargetProperty.HasMetaData(ValidEnumValuesTagName))
+		{
+			check(SourcePropertyDesc.CachedProperty)
+			TArray<FName> ValidEnumValues;
+			const FString EnumValueToInject = PropertyToInjectParentBag.GetValueSerializedString(SourcePropertyDesc.CachedProperty->GetFName()).GetValue();
+			TArray<FString> ValidEnumValuesAsString;
+			TArray<FString> InvalidEnumValuesAsString;
+			TargetProperty.GetMetaData(ValidEnumValuesTagName).ParseIntoArray(ValidEnumValuesAsString, TEXT(","));
+			TargetProperty.GetMetaData(InvalidEnumValuesTagName).ParseIntoArray(InvalidEnumValuesAsString, TEXT(","));
+
+
+			for (FString& Value : ValidEnumValuesAsString)
+			{
+				Value.TrimStartInline();
+			}
+
+			for (FString& Value : InvalidEnumValuesAsString)
+			{
+				Value.TrimStartInline();
+			}
+			
+			return (ValidEnumValuesAsString.IsEmpty() || ValidEnumValuesAsString.Contains(EnumValueToInject)) && !InvalidEnumValuesAsString.Contains(EnumValueToInject);
+		}
+
+		return true;
+	}
+
+	const bool ShouldInjectProperty(const FProperty& TargetProperty, const FPropertyBagPropertyDesc& SourcePropertyDesc, const FInstancedPropertyBag& PropertyToInjectParentBag)
+	{
+		check(SourcePropertyDesc.CachedProperty)
+		bool bShouldInject = true;
+		
+		bShouldInject &= PropertyAccessUtil::ArePropertiesCompatible(&TargetProperty, SourcePropertyDesc.CachedProperty);
+		bShouldInject &= IsValidEnumForTargetMetadata(TargetProperty, SourcePropertyDesc, PropertyToInjectParentBag);
+		bShouldInject &= TargetProperty.HasAllPropertyFlags(CPF_Edit) && !TargetProperty.HasAnyPropertyFlags(CPF_EditConst);
+		
+		return bShouldInject;
+	}
+
+#if WITH_EDITOR
+	void VisitLeafMostPropertiesNonConst(UAudioPropertiesSheetAsset& InSheetAsset, TFunction<void(FPropertyBagPropertyDesc& /*Property Desc */, UAudioPropertiesSheetAsset& /*Owning Sheet Asset*/)> OnPropertyVisited)
+	{
+		TObjectPtr<UAudioPropertiesSheetAsset> SourceSheet = &InSheetAsset;
+		const FAudioPropertiesSheet& LeafPropertiesSheet = InSheetAsset.PropertiesSheet;
+
+		while (SourceSheet != nullptr)
+		{
+			const FInstancedPropertyBag& SheetProperties = SourceSheet->PropertiesSheet.Properties;
+
+			if (const UPropertyBag* BagStruct = SheetProperties.GetPropertyBagStruct())
+			{
+				for (const FPropertyBagPropertyDesc& ConstPropertyDesc : BagStruct->GetPropertyDescs())
+				{
+					const bool bSkipOverridenProperty = SourceSheet != &InSheetAsset && LeafPropertiesSheet.IsPropertyOverridden(*ConstPropertyDesc.CachedProperty);
+
+					if (!bSkipOverridenProperty)
+					{
+						FPropertyBagPropertyDesc& PropertyDesc = const_cast<FPropertyBagPropertyDesc&>(ConstPropertyDesc);
+						OnPropertyVisited(PropertyDesc, *SourceSheet);
+					}
+
+				}
+			}
+
+			const UAudioPropertiesSheetAsset* ParentPtr = SourceSheet->PropertiesSheet.Parent;
+			SourceSheet = const_cast<UAudioPropertiesSheetAsset*>(ParentPtr);
+		}
+	}
+#endif
+}
+
+
+void FAudioPropertiesParser::GenerateNameMatchedParsingData(UE::AudioGameplay::FAudioPropertiesParsingData& ParsingData)
+{
+	if (!ensureMsgf(ParsingData.TargetClass, TEXT("Trying to generate name matching parsing data with an empty class, returning")))
+	{
+		return;
+	}
+
+	for (TFieldIterator<FProperty> PropertyIt(ParsingData.TargetClass); PropertyIt; ++PropertyIt)
+	{
+		FProperty* TargetProperty = *PropertyIt;
+		
+		if (TargetProperty)
+		{
+			const FAudioPropertiesSheet* SourceSheet = ParsingData.SourcePropertySheet;
+			if (SourceSheet)
+			{
+
+				auto OnLeafPropertyVisited = [InParsingData = &ParsingData, TargetProperty](const FPropertyBagPropertyDesc& PropertyDesc, const FAudioPropertiesSheet& OwningSheet)
+				{
+					const FName& PropertyName = TargetProperty->GetFName();
+
+					if (PropertyDesc.Name == PropertyName)
+					{
+						const bool bShouldInjectProperty = AudioPropertiesParserPrivate::ShouldInjectProperty(*TargetProperty, PropertyDesc, OwningSheet.Properties);
+
+						if (bShouldInjectProperty)
+						{
+							InParsingData->SourceToTargetPropertyBindings.Add(PropertyName, PropertyName);
+						}
+					}
+				};
+
+				AudioPropertiesUtils::VisitLeafMostProperties(*SourceSheet, OnLeafPropertyVisited);
+			}
+			else
+			{
+				const FName& PropertyName = TargetProperty->GetFName();
+				ParsingData.SourceToTargetPropertyBindings.Add(PropertyName, PropertyName);
+			}
+		}
+	}
+}
+
+void FAudioPropertiesParser::InjectPropertyValuesIntoObject(UObject& TargetObject, const UE::AudioGameplay::FAudioPropertiesParsingData& ParsingData)
+{
+	TRACE_CPUPROFILER_EVENT_SCOPE(FAudioPropertiesParser::InjectPropertyValuesIntoObject);
+
+
+	if (!ensureMsgf(TargetObject.GetClass(), TEXT("Expected a valid target object class")))
+	{
+		return;
+	}
+
+	if (!ensureMsgf(ParsingData.TargetClass, TEXT("Expected a valid target class in parsing data")))
+	{
+		return;
+	}
+
+	if (!ensureMsgf(ParsingData.SourcePropertySheet, TEXT("Expected a valid source sheet in parsing data")))
+	{
+		return;
+	}
+	
+	const FName& TargetObjectName = TargetObject.GetFName();
+
+	const FAudioPropertiesSheet* LeafSheet = ParsingData.SourcePropertySheet;
+	const FAudioPropertiesSheet* SourceSheet = ParsingData.SourcePropertySheet;
+
+	if (TargetObject.GetClass() != ParsingData.TargetClass)
+	{
+		const FName& TargetObjectClassName = TargetObject.GetClass()->GetFName();
+		const FName& ParsingDataClassName = ParsingData.TargetClass->GetFName();
+
+		UE_LOG(LogAudioProperties, Warning, TEXT("Trying to inject property values for object %s of class %s with parsing data targeting class %s. Injection aborted"), *TargetObjectName.ToString(), *TargetObjectClassName.ToString(), *ParsingDataClassName.ToString());
+		return;
+	}
+
+	bool bHasInjectedAnyProperty = false;
+	
+	while (SourceSheet != nullptr)
+	{
+		FInstancedPropertyBag PropertiesToInject;
+		PropertiesToInject.MigrateToNewBagInstance(SourceSheet->Properties);
+		
+		for (const TPair<FName, FName>& PropertyBinding : ParsingData.SourceToTargetPropertyBindings)
+		{
+			const FName& SourcePropertyName = PropertyBinding.Value;
+			const FName& TargetPropertyName = PropertyBinding.Key;
+			
+			const FPropertyBagPropertyDesc* SourcePropertyDesc = PropertiesToInject.FindPropertyDescByName(SourcePropertyName);
+			FProperty* TargetProperty = TargetObject.GetClass()->FindPropertyByName(TargetPropertyName);
+
+			if (SourcePropertyDesc && TargetProperty)
+			{
+				if (AudioPropertiesParserPrivate::IsPropertyLocallyIgnored(TargetObject, *TargetProperty))
+				{
+					continue;
+				}
+
+				if (AudioPropertiesParserPrivate::ShouldInjectProperty(*TargetProperty, *SourcePropertyDesc, PropertiesToInject))
+				{
+					check(SourcePropertyDesc->CachedProperty)
+					const bool bSkipOverridenProperty = SourceSheet != LeafSheet && LeafSheet->IsPropertyOverridden(*SourcePropertyDesc->CachedProperty);
+
+					if (!bSkipOverridenProperty)
+					{
+						FConstStructView SourcePropertyBagView = PropertiesToInject.GetValue();
+						const void* SourceAddr = SourcePropertyDesc->CachedProperty->ContainerPtrToValuePtr<void>(SourcePropertyBagView.GetMemory());
+						void* DestinationAddr = TargetProperty->ContainerPtrToValuePtr<void>(&TargetObject);
+						const FName& PropertyName = SourcePropertyDesc->CachedProperty->GetFName();
+						
+						const bool bPropertiesHaveSameValue = PropertyAccessUtil::IsCompletePropertyIdentical(SourcePropertyDesc->CachedProperty, SourceAddr, TargetProperty, DestinationAddr);
+
+						if (bPropertiesHaveSameValue)
+						{
+							UE_LOG(LogAudioProperties, Verbose, TEXT("Skipping copy property %s to object %s - value is identical"), *PropertyName.ToString(), *TargetObjectName.ToString());
+							continue;
+						}
+						
+#if WITH_EDITOR
+						AudioPropertiesParserPrivate::ApplyTargetPropertyMetadata(*TargetProperty, *SourcePropertyDesc, PropertiesToInject);
+#endif
+ 						const bool bSuccessfulCopy = PropertyAccessUtil::CopyCompletePropertyValue(SourcePropertyDesc->CachedProperty, SourceAddr, TargetProperty, DestinationAddr);
+						bHasInjectedAnyProperty |= bSuccessfulCopy;
+
+						if (!bSuccessfulCopy)
+						{
+							UE_LOG(LogAudioProperties, Verbose, TEXT("Failed to copy property %s to object %s."), *PropertyName.ToString(), *TargetObjectName.ToString());
+						}
+						else
+						{
+							UE_LOG(LogAudioProperties, Verbose, TEXT("Successfully copied property %s to object %s, marking package dirty"), *PropertyName.ToString(), *TargetObjectName.ToString());
+						}
+					}
+				}
+			}
+		}
+
+		//Inject instanced object properties
+		for(const UE::AudioGameplay::FAudioPropertiesParsingData& InstancedObjectParsingData : ParsingData.InstancedObjectsParsingData)
+		{
+			UClass* InstancedObjectClass = InstancedObjectParsingData.TargetClass;
+
+			if (!InstancedObjectClass)
+			{
+				return;
+			}
+
+			for (TFieldIterator<FProperty> PropIt(TargetObject.GetClass(), EFieldIteratorFlags::IncludeSuper); PropIt; ++PropIt)
+			{
+				FProperty* Property = *PropIt;
+
+				if(AudioPropertiesParserPrivate::IsPropertyLocallyIgnored(TargetObject, *Property))
+				{
+					continue;
+				}
+
+				auto IsRelevantObjectProperty = [](FObjectProperty* ObjectProperty, UClass* TargetObjectClass)-> bool
+				{
+					check(ObjectProperty);
+					check(TargetObjectClass);
+					
+					const bool bIsInstanced = (ObjectProperty->PropertyFlags & CPF_InstancedReference) != 0;
+
+					// A property might be relevant both if it is a child or a parent of our target
+					// We are interested in a children class when parsing properties that might be inherited
+					// We are interested in a parent when creating a new instanced object, as that could be of a child class  
+					const bool bTargetPropertyIsChild = ObjectProperty->PropertyClass->IsChildOf(TargetObjectClass);
+					const bool bTargetPropertyIsParent = TargetObjectClass->IsChildOf(ObjectProperty->PropertyClass);
+					const bool bIsRelevantClass = bTargetPropertyIsChild || bTargetPropertyIsParent;
+
+					return bIsInstanced && bIsRelevantClass;
+				};
+				
+				//Injects into an instance object property of the relevant type or creates a new one if instructed
+				if (FObjectProperty* ObjectProperty = CastField<FObjectProperty>(Property))
+				{
+					if (IsRelevantObjectProperty(ObjectProperty, InstancedObjectClass))
+					{						
+						void* PropertyValueAddress = ObjectProperty->ContainerPtrToValuePtr<void>(&TargetObject);
+						UObject* InstancedObject = ObjectProperty->GetObjectPropertyValue(PropertyValueAddress);
+								
+						if (InstancedObject && InstancedObject->IsA(InstancedObjectClass))
+						{
+							UE_LOG(LogAudioProperties, Verbose, TEXT("Injecting Properties to Instanced Object %s (%s)"), *InstancedObject->GetName(), *Property->GetName());
+							InjectPropertyValuesIntoObject(*InstancedObject, InstancedObjectParsingData);
+						}
+						
+						if(!InstancedObject && InstancedObjectParsingData.bCreateNewObjects && InstancedObjectClass->IsChildOf(ObjectProperty->PropertyClass))
+						{
+							FString NewObjectName = FString::Printf(TEXT("%s_Instanced"), *Property->GetName());
+							EObjectFlags FlagsToApply = TargetObject.GetMaskedFlags(RF_PropagateToSubObjects);
+							UObject* NewInstancedObject = NewObject<UObject>(&TargetObject, InstancedObjectClass, FName(NewObjectName), FlagsToApply);
+
+							if(NewInstancedObject)
+							{
+								InjectPropertyValuesIntoObject(*NewInstancedObject, InstancedObjectParsingData);
+								
+								ObjectProperty->SetObjectPropertyValue(PropertyValueAddress, NewInstancedObject);
+								TargetObject.MarkPackageDirty();
+
+								UE_LOG(LogAudioProperties, Verbose, TEXT("Created new instanced object %s (%s)"), *NewInstancedObject->GetName(), *Property->GetName());
+									
+							}
+						}
+					}
+				}
+				//Otherwise, check if the property is an array of instanced objects of the relevant class and does the same
+				else if (FArrayProperty* ArrayProperty = CastField<FArrayProperty>(Property))
+				{
+					if (FObjectProperty* InnerObjectProperty = CastField<FObjectProperty>(ArrayProperty->Inner))
+					{						
+						if (IsRelevantObjectProperty(InnerObjectProperty, InstancedObjectClass))
+						{
+							void* ArrayPtr = ArrayProperty->ContainerPtrToValuePtr<void>(&TargetObject);
+							FScriptArrayHelper Helper(ArrayProperty, ArrayPtr);
+							bool bHasInstancedObject = false;
+
+							for (int32 Index = 0; Index < Helper.Num(); ++Index)
+							{
+								UObject* InstancedObject = InnerObjectProperty->GetObjectPropertyValue(Helper.GetRawPtr(Index));
+								
+								if (InstancedObject && InstancedObject->IsA(InstancedObjectClass))
+								{
+									UE_LOG(LogAudioProperties, Verbose, TEXT("Injecting Properties to Instanced Object %s (%s[%d])"), *InstancedObject->GetName(), *ArrayProperty->GetName(), Index);
+									InjectPropertyValuesIntoObject(*InstancedObject, InstancedObjectParsingData);
+									bHasInstancedObject = true;
+								}
+							}
+
+							if(InstancedObjectParsingData.bCreateNewObjects && !bHasInstancedObject)
+							{
+								FString NewObjectName = FString::Printf(TEXT("%s_%d"), *Property->GetName(), Helper.Num());
+								EObjectFlags FlagsToApply = TargetObject.GetMaskedFlags(RF_PropagateToSubObjects);
+								UObject* NewInstancedObject = NewObject<UObject>(&TargetObject, InstancedObjectClass, FName(NewObjectName), FlagsToApply);
+
+								if(NewInstancedObject)
+								{
+									InjectPropertyValuesIntoObject(*NewInstancedObject, InstancedObjectParsingData);
+
+									const int32 NewIndex = Helper.AddValue();
+									void* NewElementPtr = Helper.GetRawPtr(NewIndex);
+									InnerObjectProperty->SetObjectPropertyValue(NewElementPtr, NewInstancedObject);
+									TargetObject.MarkPackageDirty();
+
+									UE_LOG(LogAudioProperties, Verbose, TEXT("Created new instanced object %s (%s[%d])"), *NewInstancedObject->GetName(), *ArrayProperty->GetName(), NewIndex);
+									
+								}
+							}
+						}
+					}
+				}
+			}
+		}
+		
+		SourceSheet = SourceSheet->Parent ? &SourceSheet->Parent->PropertiesSheet : nullptr;
+	}
+
+	if (bHasInjectedAnyProperty)
+	{
+		if (UPackage* ObjectPackage = TargetObject.GetPackage())
+		{
+			if (!ObjectPackage->IsDirty())
+			{
+				TargetObject.MarkPackageDirty();
+			}
+		}
+	}
+}
+
+#if WITH_EDITOR
+
+bool FAudioPropertiesParser::ShouldApplyTargetMetadataToProperties(const UE::AudioGameplay::FAudioPropertiesParsingData& ParsingData)
+{
+	if(!ensureMsgf(ParsingData.TargetClass, TEXT("Expected a valid target class in parsing data")))
+	{
+		return false;
+	}
+
+	if (!ensureMsgf(ParsingData.SourcePropertySheet, TEXT("Expected a valid source sheet in parsing data")))
+	{
+		return false;
+	}
+
+	bool ReturnShouldApplyMetadata = false;
+	auto CheckTargetPropertyMetadata = [&ParsingData, &ReturnShouldApplyMetadata](const FPropertyBagPropertyDesc& SheetPropertyDesc, const FAudioPropertiesSheet& OwningSheet)
+	{
+		const FName* TargetPropertyName = ParsingData.SourceToTargetPropertyBindings.FindKey(SheetPropertyDesc.Name);
+
+		if (!TargetPropertyName)
+		{
+			return;
+		}
+
+		FProperty* TargetProperty = ParsingData.TargetClass->FindPropertyByName(*TargetPropertyName);
+
+		if (!TargetProperty)
+		{
+			return;
+		}
+
+		const FInstancedPropertyBag& SheetPropertyBag = OwningSheet.Properties;
+
+		if (!PropertyAccessUtil::ArePropertiesCompatible(TargetProperty, SheetPropertyDesc.CachedProperty))
+		{
+			return;
+		}
+
+		FConstStructView SourcePropertyBagView = SheetPropertyBag.GetValue();
+		const void* SourceAddr = SheetPropertyDesc.CachedProperty->ContainerPtrToValuePtr<void>(SourcePropertyBagView.GetMemory());
+		const bool bParseNewValue = false;
+
+		const bool bShouldApplyMetadata = AudioPropertiesParserPrivate::ShouldApplyTargetPropertyMetadata(*TargetProperty, SheetPropertyDesc, SheetPropertyBag);
+		const bool bIsValidEnumForMetadata = AudioPropertiesParserPrivate::IsValidEnumForTargetMetadata(*TargetProperty, SheetPropertyDesc, SheetPropertyBag);
+
+		if (bShouldApplyMetadata || !bIsValidEnumForMetadata)
+		{
+			ReturnShouldApplyMetadata = true;
+		}
+	};
+
+	AudioPropertiesUtils::VisitLeafMostProperties(*ParsingData.SourcePropertySheet, CheckTargetPropertyMetadata);
+	
+	return ReturnShouldApplyMetadata;
+}
+
+void FAudioPropertiesParser::ApplyTargetMetadataToSheetTree(UAudioPropertiesSheetAsset& InSheetAsset, const UE::AudioGameplay::FAudioPropertiesParsingData& ParsingData)
+{
+	TRACE_CPUPROFILER_EVENT_SCOPE(FAudioPropertiesParser::ApplyTargetMetadataToSheetTree);
+
+
+	if (!ensureMsgf(ParsingData.TargetClass, TEXT("Expected a valid target class in parsing data")))
+	{
+		return;
+	}
+	
+	auto ApplyTargetPropertyMetadata = [&ParsingData](FPropertyBagPropertyDesc& SheetPropertyDesc, UAudioPropertiesSheetAsset& OwningSheet)
+	{
+		const FName* TargetPropertyName = ParsingData.SourceToTargetPropertyBindings.FindKey(SheetPropertyDesc.Name);
+
+		if (!TargetPropertyName)
+		{
+			return;
+		}
+
+		FProperty* TargetProperty = ParsingData.TargetClass->FindPropertyByName(*TargetPropertyName);
+
+		if (!TargetProperty)
+		{
+			return;
+		}
+
+		FInstancedPropertyBag& SheetPropertyBag = OwningSheet.PropertiesSheet.Properties;
+
+
+		if (!AudioPropertiesParserPrivate::ShouldInjectProperty(*TargetProperty, SheetPropertyDesc, SheetPropertyBag))
+		{
+			return;
+		}
+
+		FConstStructView SourcePropertyBagView = SheetPropertyBag.GetValue();
+		const void* SourceAddr = SheetPropertyDesc.CachedProperty->ContainerPtrToValuePtr<void>(SourcePropertyBagView.GetMemory());
+
+		if (AudioPropertiesParserPrivate::ApplyTargetPropertyMetadata(*TargetProperty, SheetPropertyDesc, SheetPropertyBag))
+		{
+			if (UPackage* ObjectPackage = OwningSheet.GetPackage())
+			{
+				if (!ObjectPackage->IsDirty())
+				{
+					OwningSheet.MarkPackageDirty();
+				}	
+			}
+		}
+	};
+
+	AudioPropertiesParserPrivate::VisitLeafMostPropertiesNonConst(InSheetAsset, ApplyTargetPropertyMetadata);
+}
+
+#endif
\ No newline at end of file
diff --git a/Engine/Plugins/Experimental/AudioProperties/Source/AudioProperties/Private/AudioPropertiesParser.h b/Engine/Plugins/Experimental/AudioProperties/Source/AudioProperties/Private/AudioPropertiesParser.h
new file mode 100644
--- /dev/null
+++ b/Engine/Plugins/Experimental/AudioProperties/Source/AudioProperties/Private/AudioPropertiesParser.h
@@ -0,0 +1,27 @@
+// Copyright Epic Games, Inc. All Rights Reserved.
+
+#pragma once
+
+class UObject; 
+class UAudioPropertiesSheetAsset;
+struct FAudioPropertiesSheet;
+
+namespace UE
+{
+	namespace AudioGameplay
+	{
+		struct FAudioPropertiesParsingData;
+	}
+}
+
+class FAudioPropertiesParser
+{
+public:
+	static void GenerateNameMatchedParsingData(UE::AudioGameplay::FAudioPropertiesParsingData& ParsingData);
+	static void InjectPropertyValuesIntoObject(UObject& TargetObject, const UE::AudioGameplay::FAudioPropertiesParsingData& ParsingData);
+	
+#if WITH_EDITOR
+	static bool ShouldApplyTargetMetadataToProperties(const UE::AudioGameplay::FAudioPropertiesParsingData& ParsingData);
+	static void ApplyTargetMetadataToSheetTree(UAudioPropertiesSheetAsset& InSheetAsset, const UE::AudioGameplay::FAudioPropertiesParsingData& ParsingData);
+#endif
+};
\ No newline at end of file
diff --git a/Engine/Plugins/Experimental/AudioProperties/Source/AudioProperties/Private/AudioPropertiesParserBindings.cpp b/Engine/Plugins/Experimental/AudioProperties/Source/AudioProperties/Private/AudioPropertiesParserBindings.cpp
new file mode 100644
--- /dev/null
+++ b/Engine/Plugins/Experimental/AudioProperties/Source/AudioProperties/Private/AudioPropertiesParserBindings.cpp
@@ -0,0 +1,50 @@
+// Copyright Epic Games, Inc. All Rights Reserved.
+
+#include "AudioPropertiesParserBindings.h"
+
+#include "AudioPropertiesBindings.h"
+#include "AudioPropertiesParser.h"
+#include "AudioPropertiesParsingData.h"
+
+#include UE_INLINE_GENERATED_CPP_BY_NAME(AudioPropertiesParserBindings)
+
+bool UAudioPropertiesParserBindings::GenerateParsingData(TObjectPtr<const UObject> TargetObject, const FAudioPropertiesSheet& SourceSheet, UE::AudioGameplay::FAudioPropertiesParsingData& InOutData) const
+{
+	if (!BindingsAsset || !TargetObject)
+	{
+		return false;
+	}
+	
+	InOutData.SourcePropertySheet = &SourceSheet;
+	InOutData.TargetClass = TargetObject->GetClass();
+	InOutData.SourceToTargetPropertyBindings = BindingsAsset->ObjectPropertyToSheetPropertyMap;
+	return true;
+}
+
+bool UAudioPropertiesParserBindings::ParseProperties(TObjectPtr<UObject> TargetObject, const FAudioPropertiesSheet& PropertiesToParse) const
+{
+	UE::AudioGameplay::FAudioPropertiesParsingData ParsingData;
+	
+	if (!GenerateParsingData(TargetObject, PropertiesToParse, ParsingData))
+	{
+		return false;
+	}
+
+	FAudioPropertiesParser::InjectPropertyValuesIntoObject(*TargetObject, ParsingData);
+	return true;
+}
+
+#if WITH_EDITOR
+
+bool UAudioPropertiesParserBindings::GenerateParsingDataFromTargetClass(UClass* TargetClass, UE::AudioGameplay::FAudioPropertiesParsingData& InOutData) const
+{
+	if (!BindingsAsset || !TargetClass)
+	{
+		return false;
+	}
+
+	InOutData.TargetClass = TargetClass;
+	InOutData.SourceToTargetPropertyBindings = BindingsAsset->ObjectPropertyToSheetPropertyMap;
+	return true;
+}
+#endif
diff --git a/Engine/Plugins/Experimental/AudioProperties/Source/AudioProperties/Private/AudioPropertiesParserNameMatch.cpp b/Engine/Plugins/Experimental/AudioProperties/Source/AudioProperties/Private/AudioPropertiesParserNameMatch.cpp
new file mode 100644
--- /dev/null
+++ b/Engine/Plugins/Experimental/AudioProperties/Source/AudioProperties/Private/AudioPropertiesParserNameMatch.cpp
@@ -0,0 +1,121 @@
+// Copyright Epic Games, Inc. All Rights Reserved.
+
+#include "AudioPropertiesParserNameMatch.h"
+
+#include "AudioPropertiesParser.h"
+#include "AudioPropertiesParsingData.h"
+
+#include UE_INLINE_GENERATED_CPP_BY_NAME(AudioPropertiesParserNameMatch)
+
+
+bool UAudioPropertiesParserNameMatch::GenerateParsingData(TObjectPtr<const UObject> TargetObject, const FAudioPropertiesSheet& SourceSheet, UE::AudioGameplay::FAudioPropertiesParsingData& InOutData) const
+{
+	if (!TargetObject)
+	{
+		return false;
+	}
+
+	
+	for(const FInstancedObjectParsingInstructions& InstancedObjectParsingRules : InstancedObjectsRules)
+	{
+		UClass* TargetClass = InstancedObjectParsingRules.InstancedObjectClass.Get();
+
+		if(!TargetClass)
+		{
+			continue;
+		}
+		
+		UE::AudioGameplay::FAudioPropertiesParsingData InstancedObjectParsingData;
+		InstancedObjectParsingData.SourcePropertySheet = &SourceSheet;
+		InstancedObjectParsingData.TargetClass = TargetClass;
+		InstancedObjectParsingData.bCreateNewObjects = InstancedObjectParsingRules.bCreateIfMissing;
+		
+		if(!InstancedObjectParsingRules.bOverrideParsing)
+		{
+			GenerateParsingDataFromTargetClass(TargetClass, InstancedObjectParsingData);
+		}
+		else
+		{
+			if(InstancedObjectParsingRules.PropertiesParser)
+			{
+				TObjectPtr<UObject> ParsingDataGenerationObject = NewObject<UObject>(GetOuter(), InstancedObjectParsingRules.InstancedObjectClass.Get());
+				InstancedObjectParsingRules.PropertiesParser->GenerateParsingData(ParsingDataGenerationObject, SourceSheet, InstancedObjectParsingData);
+			}
+		}
+		
+		InOutData.InstancedObjectsParsingData.Add(MoveTemp(InstancedObjectParsingData));	
+	}
+
+	InOutData.SourcePropertySheet = &SourceSheet;
+	UClass* ObjectClass = TargetObject->GetClass();
+	
+	return GenerateParsingDataFromTargetClass(ObjectClass, InOutData);
+}
+
+bool UAudioPropertiesParserNameMatch::ParseProperties(TObjectPtr<UObject> TargetObject, const FAudioPropertiesSheet& PropertiesToParse) const
+{
+	UE::AudioGameplay::FAudioPropertiesParsingData ParsingData;
+	
+	if (!GenerateParsingData(TargetObject, PropertiesToParse, ParsingData))
+	{
+		return false;
+	}
+
+	FAudioPropertiesParser::InjectPropertyValuesIntoObject(*TargetObject, ParsingData);
+	return true;	
+}
+
+#if WITH_EDITOR
+
+void UAudioPropertiesParserNameMatch::FitPropertiesOnAssetTree(UAudioPropertiesSheetAsset& InSheetAsset) const
+{
+	if (!ValidationClass)
+	{
+		return;
+	}
+
+	UE::AudioGameplay::FAudioPropertiesParsingData ParsingData;
+	
+	if (!GenerateParsingDataFromTargetClass(ValidationClass, ParsingData))
+	{
+		return;
+	}
+
+	FAudioPropertiesParser::ApplyTargetMetadataToSheetTree(InSheetAsset, ParsingData);
+}
+
+
+bool UAudioPropertiesParserNameMatch::ValidatePropertiesOnAssetTree(const UAudioPropertiesSheetAsset& InSheetAsset) const
+{
+	if (!ValidationClass)
+	{
+		return true;
+	}
+
+	UE::AudioGameplay::FAudioPropertiesParsingData ParsingData;
+	ParsingData.SourcePropertySheet = &InSheetAsset.PropertiesSheet;
+	
+	if (!GenerateParsingDataFromTargetClass(ValidationClass, ParsingData))
+	{
+		return false;
+	}
+
+	const bool bPropertiesNeedMetadata = FAudioPropertiesParser::ShouldApplyTargetMetadataToProperties(ParsingData);
+
+	return !bPropertiesNeedMetadata;
+}
+
+bool UAudioPropertiesParserNameMatch::GenerateParsingDataFromTargetClass(UClass* TargetClass, UE::AudioGameplay::FAudioPropertiesParsingData& InOutData) const
+{
+	if (!TargetClass)
+	{
+		return false;
+	}
+
+	InOutData.TargetClass = TargetClass;
+
+	FAudioPropertiesParser::GenerateNameMatchedParsingData(InOutData);
+	return true;
+}
+
+#endif
diff --git a/Engine/Plugins/Experimental/AudioProperties/Source/AudioProperties/Private/AudioPropertiesSheet.cpp b/Engine/Plugins/Experimental/AudioProperties/Source/AudioProperties/Private/AudioPropertiesSheet.cpp
new file mode 100644
--- /dev/null
+++ b/Engine/Plugins/Experimental/AudioProperties/Source/AudioProperties/Private/AudioPropertiesSheet.cpp
@@ -0,0 +1,514 @@
+// Copyright Epic Games, Inc. All Rights Reserved.
+
+#include "AudioPropertiesSheet.h"
+
+#include "AudioPropertiesBindings.h"
+#include "AudioPropertiesLogs.h"
+#include "AudioPropertiesParserBase.h"
+#include "AudioPropertiesParsingData.h"
+#include "AudioPropertiesUtils.h"
+#include "Interfaces/ITargetPlatform.h"
+#include "UObject/ObjectSaveContext.h"
+#include "UObject/PropertyAccessUtil.h"
+
+#if WITH_EDITOR
+#include "Misc/DataValidation.h"
+#endif
+
+#include UE_INLINE_GENERATED_CPP_BY_NAME(AudioPropertiesSheet)
+
+#define LOCTEXT_NAMESPACE "AudioPropertiesSheetAsset"
+
+namespace AudioPropertySheetPrivate
+{
+	bool SameType(const FPropertyBagPropertyDesc& Property1, const FPropertyBagPropertyDesc& Property2)
+	{
+		const bool bDifferentType = Property1.ValueType != Property2.ValueType || Property1.ValueTypeObject != Property2.ValueTypeObject;
+		return !bDifferentType;
+	}
+
+	bool IsRelevantInstancedObjectProperty(const FProperty* TargetProperty, const UE::AudioGameplay::FAudioPropertiesParsingData& InstancedObjectParsingData)
+	{
+		check(TargetProperty)
+		
+		UClass* TargetObjectClass = InstancedObjectParsingData.TargetClass;
+
+		if(!TargetObjectClass)
+		{
+			return false;
+		}
+		
+		const FArrayProperty* ArrayProperty = CastField<FArrayProperty>(TargetProperty);
+		const FObjectProperty* ObjectProperty = ArrayProperty != nullptr ? CastField<FObjectProperty>(ArrayProperty->Inner) : CastField<FObjectProperty>(TargetProperty);
+
+		if(!ObjectProperty)
+		{
+			return false;
+		}
+
+		const bool bIsInstanced = (ObjectProperty->PropertyFlags & CPF_InstancedReference) != 0;
+		const bool bIsRelevantClass = InstancedObjectParsingData.bCreateNewObjects ? TargetObjectClass->IsChildOf(ObjectProperty->PropertyClass) : ObjectProperty->PropertyClass->IsChildOf(TargetObjectClass); 
+
+		return bIsInstanced && bIsRelevantClass;
+	}
+}
+
+bool FAudioPropertiesSheet::UpdatePropertyOverride(const FProperty& InProperty, bool bMarkAsOverridden)
+{
+	if (bMarkAsOverridden)
+	{
+
+		if (IsPropertyOverridden(InProperty))
+		{
+			return true;
+		}
+
+		AudioPropertiesSheet::PropertyDescParentPair PropertyDescParentPair = FindClosestParentWithProperty(InProperty);
+
+		if (!PropertyDescParentPair.Key || !PropertyDescParentPair.Value)
+		{
+			return false;
+ 		}
+
+		const FInstancedPropertyBag& ParentParams = PropertyDescParentPair.Value->PropertiesSheet.Properties;
+		const FPropertyBagPropertyDesc* PropertyDescToInherit = PropertyDescParentPair.Key;
+
+		//Make property local by cloning it into a new desc 
+		const EPropertyBagContainerType ContainerType = PropertyDescToInherit->ContainerTypes.GetFirstContainerType();
+		FPropertyBagPropertyDesc NewPropertyDesc(InProperty.GetFName(), ContainerType, PropertyDescToInherit->ValueType, PropertyDescToInherit->ValueTypeObject.Get());
+		Properties.AddProperties({ NewPropertyDesc });
+ 		
+ 		//Copy value from parent property
+		const FPropertyBagPropertyDesc* OverridenPropertyDesc = Properties.FindPropertyDescByName(InProperty.GetFName());
+		const void* SourceAddr = PropertyDescToInherit->CachedProperty->ContainerPtrToValuePtr<void>(ParentParams.GetValue().GetMemory());
+		void* DestinationAddr = OverridenPropertyDesc->CachedProperty->ContainerPtrToValuePtr<void>(Properties.GetMutableValue().GetMemory());
+		OverridenPropertyDesc->CachedProperty->CopyCompleteValue(DestinationAddr, SourceAddr);
+		
+		OnAudioPropertyOverrideChange.Broadcast(InProperty, bMarkAsOverridden);
+
+		return true;
+	}
+	else
+	{
+		Properties.RemovePropertyByName(InProperty.GetFName());
+
+		OnAudioPropertyOverrideChange.Broadcast(InProperty, bMarkAsOverridden);
+
+		return true;
+	}
+ }
+
+bool FAudioPropertiesSheet::IsPropertyOverridden(const FProperty& InProperty) const
+{
+	const FPropertyBagPropertyDesc* LocalPropertyDesc = Properties.FindPropertyDescByName(InProperty.GetFName());
+	const FPropertyBagPropertyDesc* ParentPropertyDesc = FindClosestParentWithProperty(InProperty).Key;
+	
+	if (LocalPropertyDesc && ParentPropertyDesc)
+	{
+		return AudioPropertySheetPrivate::SameType(*LocalPropertyDesc, *ParentPropertyDesc);
+	}
+
+	return false;
+}
+
+AudioPropertiesSheet::PropertyDescParentPair FAudioPropertiesSheet::FindClosestParentWithProperty(const FProperty& InProperty) const
+{
+	if (!Parent)
+	{
+		return AudioPropertiesSheet::PropertyDescParentPair(nullptr, nullptr);
+	}
+
+	TObjectPtr<const UAudioPropertiesSheetAsset> VisitedParent = Parent;
+
+	while (VisitedParent)
+	{
+		if (const FPropertyBagPropertyDesc* ParentDesc = VisitedParent->PropertiesSheet.Properties.FindPropertyDescByName(InProperty.GetFName()))
+		{
+			return AudioPropertiesSheet::PropertyDescParentPair(ParentDesc, VisitedParent);
+		}
+
+		VisitedParent = VisitedParent->PropertiesSheet.Parent;
+	}
+
+	return AudioPropertiesSheet::PropertyDescParentPair(nullptr, nullptr);;
+}
+
+
+void FAudioPropertiesSheet::ReconcileProperties()
+{
+	if (!Properties.GetPropertyBagStruct())
+	{
+		return;
+	}
+
+	TSet<FGuid> MatchingOverrides;
+	TArray<FPropertyBagPropertyDesc> NewDescs;
+	TArray<TPair<const FPropertyBagPropertyDesc&, const FInstancedPropertyBag&>> PropertiesToCopyValueFrom;
+	
+	//bag to store old properties values
+	FInstancedPropertyBag OldProperties;
+
+	for (const FPropertyBagPropertyDesc& LocalPropertyDesc : Properties.GetPropertyBagStruct()->GetPropertyDescs())
+	{
+		check(LocalPropertyDesc.CachedProperty)
+		
+		// Check if a parent sheet contains the property
+		AudioPropertiesSheet::PropertyDescParentPair PropertyDescParentPair = FindClosestParentWithProperty(*LocalPropertyDesc.CachedProperty);
+
+		if (PropertyDescParentPair.Value)
+		{
+			const FPropertyBagPropertyDesc* ParentPropertyDesc = PropertyDescParentPair.Key;
+
+			if (ParentPropertyDesc)
+			{
+				if (!AudioPropertySheetPrivate::SameType(LocalPropertyDesc, *ParentPropertyDesc))
+				{
+					UE_LOG(LogAudioProperties, Warning, TEXT("Parent type changed for property %s, property type will be changed and value might be lost"), *LocalPropertyDesc.CachedProperty->GetFName().ToString());
+					NewDescs.Emplace(LocalPropertyDesc.CachedProperty->GetFName(), ParentPropertyDesc->ContainerTypes.GetFirstContainerType(), ParentPropertyDesc->ValueType, ParentPropertyDesc->ValueTypeObject.Get());
+					
+					if (PropertyAccessUtil::ArePropertiesCompatible(LocalPropertyDesc.CachedProperty, ParentPropertyDesc->CachedProperty))
+					{
+						//if the local property is compatible with the new type we keep its value
+						PropertiesToCopyValueFrom.Emplace(LocalPropertyDesc, OldProperties);
+					}
+					else
+					{
+						//otherwise we initialize using the parent value
+						PropertiesToCopyValueFrom.Emplace(*ParentPropertyDesc, PropertyDescParentPair.Value->PropertiesSheet.Properties);
+					}
+					continue;
+
+				}
+			}
+		}
+
+		NewDescs.Add(LocalPropertyDesc);
+	}
+	
+	OldProperties.MigrateToNewBagInstance(Properties);
+
+	//update params from reconciled descriptions
+	const UPropertyBag* NewBag = UPropertyBag::GetOrCreateFromDescs(NewDescs);
+	Properties.MigrateToNewBagStruct(NewBag);
+
+	//now that we instantiated the underlying UObject holding FProperties, copy old values
+	for (TPair<const FPropertyBagPropertyDesc&, const FInstancedPropertyBag&> PropertyToCopy : PropertiesToCopyValueFrom)
+	{
+		const void* SourceAddr = PropertyToCopy.Key.CachedProperty->ContainerPtrToValuePtr<void>(PropertyToCopy.Value.GetValue().GetMemory());
+
+		const FPropertyBagPropertyDesc* TargetPropertyDesc = NewBag->FindPropertyDescByName(PropertyToCopy.Key.CachedProperty->GetFName());
+		void* TargetAddr = TargetPropertyDesc->CachedProperty->ContainerPtrToValuePtr<void>(Properties.GetMutableValue().GetMemory());
+		
+		PropertyAccessUtil::CopyCompletePropertyValue(PropertyToCopy.Key.CachedProperty, SourceAddr, TargetPropertyDesc->CachedProperty, TargetAddr);
+	}
+}
+
+void FAudioPropertiesSheet::OnPreSave(const FObjectPreSaveContext& SaveContext)
+{
+	ReconcileProperties();
+
+	if(!SaveContext.IsCooking())
+	{
+		return;
+	}
+
+	if (!SaveContext.GetTargetPlatform()->HasEditorOnlyData())
+	{
+		return;
+	}
+
+	FlattenSheet();
+}
+
+void FAudioPropertiesSheet::FlattenSheet()
+{
+	if (!Properties.GetPropertyBagStruct())
+	{
+		return;
+	}
+
+	ReconcileProperties();
+
+	auto ParseLeafMostNonLocalProperties = [this, &LocalPropertyBag = Properties](const FPropertyBagPropertyDesc& PropertyDesc, const FAudioPropertiesSheet& OwningSheet)
+	{
+		if (this == &OwningSheet)
+		{
+			return;
+		}
+
+		if (!PropertyDesc.CachedProperty)
+		{
+			return;
+		}
+
+		LocalPropertyBag.AddProperty(PropertyDesc.CachedProperty->GetFName(), PropertyDesc.CachedProperty);
+
+		const void* SourcePropertyAddr = PropertyDesc.CachedProperty->ContainerPtrToValuePtr<void>(OwningSheet.Properties.GetValue().GetMemory());
+		LocalPropertyBag.SetValue(PropertyDesc.CachedProperty->GetFName(), PropertyDesc.CachedProperty, SourcePropertyAddr);
+
+	};
+
+	AudioPropertiesUtils::VisitLeafMostProperties(*this, ParseLeafMostNonLocalProperties);
+
+	Parent = nullptr;
+}
+
+UAudioPropertiesSheetAsset::UAudioPropertiesSheetAsset(const FObjectInitializer& ObjectInitializer)
+	: Super(ObjectInitializer)
+{
+}
+
+#if WITH_EDITOR
+
+void UAudioPropertiesSheetAsset::PreSave(FObjectPreSaveContext SaveContext)
+{
+	Super::PreSave(SaveContext);
+	
+	PropertiesSheet.OnPreSave(SaveContext);
+}
+
+EDataValidationResult UAudioPropertiesSheetAsset::IsDataValid(FDataValidationContext& Context) const
+{
+	Super::IsDataValid(Context);
+
+	FString AssetName;
+	GetName(AssetName);
+
+	UE_LOG(LogAudioProperties, Log, TEXT("Validating Property Sheet %s"), *AssetName);
+
+	if (PropertiesParser && !PropertiesParser->ValidatePropertiesOnAssetTree(*this))
+	{
+		Context.AddWarning(LOCTEXT("Parser Validation failed", "Parser validation failed, check output log"));
+
+		UE_LOG(LogAudioProperties, Log, TEXT("Validating Property Sheet %s - Validation failed"), *AssetName);
+
+		return EDataValidationResult::Invalid;
+	}
+
+	UE_LOG(LogAudioProperties, Log, TEXT("Validating Property Sheet %s - Validation successfull"), *AssetName);
+
+	return EDataValidationResult::Valid;
+}
+
+void UAudioPropertiesSheetAsset::FitPropertiesForValidation()
+{
+	TObjectPtr<const UAudioPropertiesParserBase> ParserToUse = FindClosestParserInInheritanceTree().Key;
+
+	if (!ParserToUse)
+	{
+		UE_LOG(LogAudioProperties, Log, TEXT("Trying to validate data with null parser"))
+		return;
+	}
+	
+	ParserToUse->FitPropertiesOnAssetTree(*this);
+
+}
+
+bool UAudioPropertiesSheetAsset::CopyToObjectProperties(UObject* TargetObject) const
+{
+	TObjectPtr<const UAudioPropertiesParserBase> ParserToUse = FindClosestParserInInheritanceTree().Key;
+
+	if (!ParserToUse)
+	{
+
+		FString AssetName;
+		GetName(AssetName);
+
+		UE_LOG(LogAudioProperties, Log, TEXT("Failed to inject data from property sheet %s - could not find a valid parser"), *AssetName)
+		return false;
+	}
+
+	return ParserToUse->ParseProperties(TargetObject, PropertiesSheet);
+
+}
+
+FDelegateHandle UAudioPropertiesSheetAsset::BindPropertiesCopyToSheetChanges(UObject* TargetObject)
+{
+	auto OnPropertySheetChanged = [this, TargetObject](AudioPropertiesSheet::PostEditChangeType ChangeType) 
+		{ 
+			this->CopyToObjectProperties(TargetObject); 
+		};
+	
+	return OnPostEditChange.AddWeakLambda(TargetObject, OnPropertySheetChanged);
+}
+
+void UAudioPropertiesSheetAsset::UnbindCopyFromPropertySheetChanges(UObject* ObjectToUnbind)
+{
+	OnPostEditChange.RemoveAll(ObjectToUnbind);
+}
+
+void UAudioPropertiesSheetAsset::PreEditChange(FProperty* PropertyAboutToChange)
+{
+	Super::PreEditChange(PropertyAboutToChange);
+
+	if (!PropertyAboutToChange)
+	{
+		return;
+	}
+
+	if (PropertyAboutToChange->GetName() == GET_MEMBER_NAME_CHECKED(FAudioPropertiesSheet, Parent))
+	{
+		TObjectPtr<const UAudioPropertiesSheetAsset> ParentValue = nullptr;
+		PropertyAboutToChange->GetValue_InContainer(&PropertiesSheet, &ParentValue);
+
+		if (ParentValue)
+		{
+			UAudioPropertiesSheetAsset* ParentPtr = const_cast<UAudioPropertiesSheetAsset*>(ParentValue.Get());
+			ParentPtr->OnPostEditChange.RemoveAll(this);
+		}
+	}
+}
+
+void UAudioPropertiesSheetAsset::PostEditChangeProperty(struct FPropertyChangedEvent& PropertyChangedEvent)
+{
+	Super::PostEditChangeProperty(PropertyChangedEvent);
+
+	AudioPropertiesSheet::PostEditChangeType BroadcastChangeType = AudioPropertiesSheet::PostEditChangeType::Other;
+
+	if (PropertyChangedEvent.GetPropertyName() == GET_MEMBER_NAME_CHECKED(FAudioPropertiesSheet, Parent))
+	{
+		TObjectPtr<const UAudioPropertiesSheetAsset> ParentValue = nullptr;
+		TObjectPtr<const UAudioPropertiesSheetAsset> AssetBeingChecked = this;
+		PropertyChangedEvent.Property->GetValue_InContainer(&PropertiesSheet, &ParentValue);
+
+		if (ParentValue)
+		{
+			UAudioPropertiesSheetAsset* ParentPtr = const_cast<UAudioPropertiesSheetAsset*>(ParentValue.Get());
+			ParentPtr->OnPostEditChange.AddWeakLambda(this, [this](AudioPropertiesSheet::PostEditChangeType ChangeType) {this->OnParentPostEditChange(ChangeType); });
+		}
+
+		while (ParentValue)
+		{
+			checkf(ParentValue != this, TEXT("Found asset %s while walking its parent chain at %s"), *this->GetName(), *AssetBeingChecked->GetName());
+			AssetBeingChecked = ParentValue;
+			ParentValue = ParentValue->PropertiesSheet.Parent;
+		}
+
+		PropertiesSheet.ReconcileProperties();
+		BroadcastChangeType = AudioPropertiesSheet::PostEditChangeType::ParentChange;
+	}
+	else if (PropertyChangedEvent.GetPropertyName() == GET_MEMBER_NAME_CHECKED(FAudioPropertiesSheet, Properties))
+	{
+		PropertiesSheet.ReconcileProperties();
+		BroadcastChangeType = AudioPropertiesSheet::PostEditChangeType::PropertyBag;
+
+	}
+	else if (PropertyChangedEvent.GetPropertyName() == GET_MEMBER_NAME_CHECKED(UAudioPropertiesSheetAsset, PropertiesParser))
+	{
+		BroadcastChangeType = AudioPropertiesSheet::PostEditChangeType::PropertiesParser;
+	}
+
+	OnPostEditChange.Broadcast(BroadcastChangeType);
+}
+
+void UAudioPropertiesSheetAsset::PostLoad()
+{
+	Super::PostLoad();
+
+	PropertiesSheet.ReconcileProperties();
+
+	if (PropertiesSheet.Parent)
+	{
+		UAudioPropertiesSheetAsset* ParentPtr = const_cast<UAudioPropertiesSheetAsset*>(PropertiesSheet.Parent.Get());
+		ParentPtr->OnPostEditChange.AddWeakLambda(this, [this](AudioPropertiesSheet::PostEditChangeType ChangeType) 
+			{
+				this->OnParentPostEditChange(ChangeType); 
+			}
+		);
+	}
+}
+
+#endif
+
+void UAudioPropertiesSheetAsset::OnParentPostEditChange(const AudioPropertiesSheet::PostEditChangeType& ChangeType)
+{
+	switch (ChangeType)
+	{
+	case AudioPropertiesSheet::PostEditChangeType::PropertyBag:
+	case AudioPropertiesSheet::PostEditChangeType::ParentChange:
+		PropertiesSheet.ReconcileProperties();
+		MarkPackageDirty();
+		break;
+	case AudioPropertiesSheet::PostEditChangeType::Other:
+		break;
+	default:
+		break;
+	}
+
+	OnPostEditChange.Broadcast(ChangeType);
+}
+
+AudioPropertiesSheet::ParserParentPair UAudioPropertiesSheetAsset::FindClosestParserInInheritanceTree() const
+{
+	if (PropertiesParser)
+	{
+		return AudioPropertiesSheet::ParserParentPair(PropertiesParser, this);
+	}
+
+	TObjectPtr<const UAudioPropertiesParserBase> ClosestParser = nullptr;
+	const UAudioPropertiesSheetAsset* ParentToVisit = PropertiesSheet.Parent; 
+	const UAudioPropertiesSheetAsset* OutParent = PropertiesSheet.Parent; 
+
+
+	while (ParentToVisit && ClosestParser == nullptr)
+	{
+		OutParent = ParentToVisit;
+		ClosestParser = ParentToVisit->PropertiesParser;
+		ParentToVisit = ParentToVisit->PropertiesSheet.Parent;
+	}
+
+	return AudioPropertiesSheet::ParserParentPair(ClosestParser, ClosestParser ? OutParent : nullptr);
+}
+
+void UAudioPropertiesSheetAsset::GetTargetedPropertyNames(TObjectPtr<const UObject> TargetObject, TArray<FName>& OutProperties) const
+{
+	TObjectPtr<const UAudioPropertiesParserBase> Parser = FindClosestParserInInheritanceTree().Key;
+
+	if (!Parser)
+	{
+		return;
+	}
+
+	UE::AudioGameplay::FAudioPropertiesParsingData ParsingData;
+	Parser->GenerateParsingData(TargetObject, PropertiesSheet, ParsingData);
+
+	auto OnLeafPropertyVisited = [OutPropertyNames = &OutProperties, Bindings = &ParsingData.SourceToTargetPropertyBindings](const FPropertyBagPropertyDesc& PropertyDesc, const FAudioPropertiesSheet& OwningSheet) 
+	{
+		if (const FName* TargetPropertyName = Bindings->Find(PropertyDesc.Name))
+		{
+			OutPropertyNames->Add(*TargetPropertyName);
+		}
+	};
+
+	AudioPropertiesUtils::VisitLeafMostProperties(PropertiesSheet, OnLeafPropertyVisited);
+
+	for(const UE::AudioGameplay::FAudioPropertiesParsingData& InstancedObjectParsingData : ParsingData.InstancedObjectsParsingData)
+	{
+		if(InstancedObjectParsingData.SourceToTargetPropertyBindings.Num() > 0 || InstancedObjectParsingData.bCreateNewObjects)
+		{
+			UClass* InstancedObjectClass = InstancedObjectParsingData.TargetClass;
+
+			if(!InstancedObjectClass)
+			{
+				continue;
+			}
+
+			for (TFieldIterator<FProperty> PropIt(TargetObject->GetClass()); PropIt; ++PropIt)
+			{
+				const FProperty* Property = *PropIt;
+
+				if(!ensureAlwaysMsgf(Property, TEXT("Fsound a null property from a field iterator")))
+				{
+					continue;
+				}
+				
+				if (AudioPropertySheetPrivate::IsRelevantInstancedObjectProperty(Property, InstancedObjectParsingData))
+				{
+					OutProperties.Add(Property->GetFName());
+				}
+			}
+		}
+	}
+}
+
+#undef LOCTEXT_NAMESPACE
diff --git a/Engine/Plugins/Experimental/AudioProperties/Source/AudioProperties/Private/AudioPropertiesUtils.cpp b/Engine/Plugins/Experimental/AudioProperties/Source/AudioProperties/Private/AudioPropertiesUtils.cpp
new file mode 100644
--- /dev/null
+++ b/Engine/Plugins/Experimental/AudioProperties/Source/AudioProperties/Private/AudioPropertiesUtils.cpp
@@ -0,0 +1,87 @@
+// Copyright Epic Games, Inc. All Rights Reserved.
+
+#include "AudioPropertiesUtils.h"
+
+#include "AudioPropertiesSheet.h"
+
+#include "AssetRegistry/AssetRegistryModule.h"
+
+bool AudioPropertiesUtils::IsSheetInInheritanceTree(const UAudioPropertiesSheetAsset& InLeafSheet, const TObjectPtr<const UAudioPropertiesSheetAsset> InSheetToFind)
+{
+	check(InSheetToFind)
+
+	TObjectPtr<const UAudioPropertiesSheetAsset> ParentValue = InLeafSheet.PropertiesSheet.Parent;
+
+	while (ParentValue)
+	{
+		if (ParentValue == InSheetToFind)
+		{
+			return true;
+		}
+
+		ParentValue = ParentValue->PropertiesSheet.Parent;
+	}
+
+	return false;
+}
+
+void AudioPropertiesUtils::GetAudioPropertiesSheetAssetReferencers(const UAudioPropertiesSheetAsset& InAsset, TArray<FAssetData>& OutReferencersData, const UClass* ClassFilter /*= nullptr*/)
+{
+	FAssetRegistryModule& AssetRegistryModule = FModuleManager::LoadModuleChecked<FAssetRegistryModule>(TEXT("AssetRegistry"));
+	IAssetRegistry& AssetRegistry = AssetRegistryModule.Get();
+	TArray<FAssetIdentifier> AssetDependencies;
+	FARFilter AssetRegistryFilter;
+
+	if (ClassFilter)
+	{
+		FTopLevelAssetPath ClassPathName = ClassFilter->GetClassPathName();
+		AssetRegistryFilter.ClassPaths.Add(ClassPathName);
+
+	}
+
+	AssetRegistry.GetReferencers(InAsset.GetPackage()->GetFName(), AssetRegistryFilter.PackageNames);
+	AssetRegistry.GetAssets(AssetRegistryFilter, OutReferencersData);
+}
+
+void AudioPropertiesUtils::VisitAudioPropertiesSheetAssetReferencers(const UAudioPropertiesSheetAsset& InAsset, TFunctionRef<void(FAssetData&)> OnAssetVisited, const UClass* ClassFilter /*= nullptr*/)
+{
+	TArray<FAssetData> ReferencersData;
+	GetAudioPropertiesSheetAssetReferencers(InAsset, ReferencersData, ClassFilter);
+
+	for (FAssetData& AssetData : ReferencersData)
+	{
+		OnAssetVisited(AssetData);
+	}
+}
+
+void AudioPropertiesUtils::VisitLeafMostProperties(const UAudioPropertiesSheetAsset& InSheetAsset, TFunction<void(const FPropertyBagPropertyDesc&, const FAudioPropertiesSheet&)> OnPropertyVisited)
+{
+	VisitLeafMostProperties(InSheetAsset.PropertiesSheet, OnPropertyVisited);
+}
+
+void AudioPropertiesUtils::VisitLeafMostProperties(const FAudioPropertiesSheet& InSheetAsset, TFunction<void(const FPropertyBagPropertyDesc&, const FAudioPropertiesSheet&)> OnPropertyVisited)
+{
+	const FAudioPropertiesSheet* LeafSheet = &InSheetAsset;
+	const FAudioPropertiesSheet* SourceSheet = &InSheetAsset;
+
+	while (SourceSheet != nullptr)
+	{
+		const FInstancedPropertyBag& SheetProperties = SourceSheet->Properties;
+
+		if (const UPropertyBag* BagStruct = SheetProperties.GetPropertyBagStruct())
+		{
+			for (const FPropertyBagPropertyDesc& PropertyDesc : BagStruct->GetPropertyDescs())
+			{
+				const bool bSkipOverridenProperty = SourceSheet != &InSheetAsset && LeafSheet->IsPropertyOverridden(*PropertyDesc.CachedProperty);
+
+				if (!bSkipOverridenProperty)
+				{
+					OnPropertyVisited(PropertyDesc, *SourceSheet);
+				}
+
+			}
+		}
+
+		SourceSheet = SourceSheet->Parent ? &SourceSheet->Parent->PropertiesSheet : nullptr;;
+	}
+}
diff --git a/Engine/Plugins/Experimental/AudioProperties/Source/AudioProperties/Public/AudioPropertiesBindings.h b/Engine/Plugins/Experimental/AudioProperties/Source/AudioProperties/Public/AudioPropertiesBindings.h
new file mode 100644
--- /dev/null
+++ b/Engine/Plugins/Experimental/AudioProperties/Source/AudioProperties/Public/AudioPropertiesBindings.h
@@ -0,0 +1,19 @@
+// Copyright Epic Games, Inc. All Rights Reserved.
+
+#pragma once
+
+#include "UObject/Object.h"
+
+#include "AudioPropertiesBindings.generated.h"
+
+UCLASS(MinimalAPI)
+class UAudioPropertiesBindings : public UObject
+{
+	GENERATED_BODY()
+
+	UAudioPropertiesBindings() = default;
+
+public:
+	UPROPERTY(EditDefaultsOnly, Category = "Default")
+	TMap<FName, FName> ObjectPropertyToSheetPropertyMap;
+};
diff --git a/Engine/Plugins/Experimental/AudioProperties/Source/AudioProperties/Public/AudioPropertiesLogs.h b/Engine/Plugins/Experimental/AudioProperties/Source/AudioProperties/Public/AudioPropertiesLogs.h
new file mode 100644
--- /dev/null
+++ b/Engine/Plugins/Experimental/AudioProperties/Source/AudioProperties/Public/AudioPropertiesLogs.h
@@ -0,0 +1,7 @@
+// Copyright Epic Games, Inc. All Rights Reserved.
+
+#pragma once
+
+#include "Logging/LogMacros.h"
+
+DECLARE_LOG_CATEGORY_EXTERN(LogAudioProperties, Log, All);
diff --git a/Engine/Plugins/Experimental/AudioProperties/Source/AudioProperties/Public/AudioPropertiesModule.h b/Engine/Plugins/Experimental/AudioProperties/Source/AudioProperties/Public/AudioPropertiesModule.h
new file mode 100644
--- /dev/null
+++ b/Engine/Plugins/Experimental/AudioProperties/Source/AudioProperties/Public/AudioPropertiesModule.h
@@ -0,0 +1,13 @@
+// Copyright Epic Games, Inc. All Rights Reserved.
+
+#pragma once
+
+#include "Modules/ModuleInterface.h"
+
+class FAudioPropertiesModule : public IModuleInterface
+{
+public:
+	virtual void StartupModule() override;
+	virtual void ShutdownModule() override;
+};
+
diff --git a/Engine/Plugins/Experimental/AudioProperties/Source/AudioProperties/Public/AudioPropertiesParserBase.h b/Engine/Plugins/Experimental/AudioProperties/Source/AudioProperties/Public/AudioPropertiesParserBase.h
new file mode 100644
--- /dev/null
+++ b/Engine/Plugins/Experimental/AudioProperties/Source/AudioProperties/Public/AudioPropertiesParserBase.h
@@ -0,0 +1,58 @@
+// Copyright Epic Games, Inc. All Rights Reserved.
+
+#pragma once
+
+#include "AudioPropertiesSheet.h"
+#include "Templates/SubclassOf.h"
+#include "UObject/Object.h"
+
+#include "AudioPropertiesParserBase.generated.h"
+
+struct FAudioPropertiesSheet;
+
+namespace UE
+{
+	namespace AudioGameplay
+	{
+		struct FAudioPropertiesParsingData;
+	}
+}
+
+USTRUCT()
+struct FInstancedObjectParsingInstructions
+{
+	GENERATED_BODY()
+
+	//The instanced object class that will be targeted by these instructions
+	UPROPERTY(EditAnywhere, Category = "Parsing Instructions")
+	TSubclassOf<UObject> InstancedObjectClass;
+
+	//If the target instanced uobject is missing, create one with the properties that should have been parsed
+	UPROPERTY(EditAnywhere, Category = "Parsing Instructions")
+	bool bCreateIfMissing = false;
+
+	//Use this to signal that you would like to use a different parser for the instanced objects
+	UPROPERTY(EditAnywhere, Category = "Parsing Instructions")
+	bool bOverrideParsing = false;
+	
+	UPROPERTY(EditAnywhere, Instanced, Category = "Parsing Instructions", meta = (EditCondition = "bOverrideParsing", EditConditionHides))
+	TObjectPtr<UAudioPropertiesParserBase> PropertiesParser;
+};
+
+
+
+UCLASS(Abstract, EditInlineNew, MinimalAPI)
+class UAudioPropertiesParserBase : public UObject
+{
+	GENERATED_BODY()
+
+public:
+	virtual bool GenerateParsingData(TObjectPtr<const UObject> TargetObject,  const FAudioPropertiesSheet& SourceSheet, UE::AudioGameplay::FAudioPropertiesParsingData& InOutData) const { return false; };
+	virtual bool ParseProperties(TObjectPtr<UObject> TargetObject, const FAudioPropertiesSheet& PropertiesToParse) const { return false; };
+
+#if WITH_EDITOR
+	virtual bool ValidatePropertiesOnAssetTree(const UAudioPropertiesSheetAsset& InSheetAsset) const { return false; }
+	virtual void FitPropertiesOnAssetTree(UAudioPropertiesSheetAsset& InSheetAsset) const { return; }
+#endif
+
+};
\ No newline at end of file
diff --git a/Engine/Plugins/Experimental/AudioProperties/Source/AudioProperties/Public/AudioPropertiesParserBindings.h b/Engine/Plugins/Experimental/AudioProperties/Source/AudioProperties/Public/AudioPropertiesParserBindings.h
new file mode 100644
--- /dev/null
+++ b/Engine/Plugins/Experimental/AudioProperties/Source/AudioProperties/Public/AudioPropertiesParserBindings.h
@@ -0,0 +1,39 @@
+// Copyright Epic Games, Inc. All Rights Reserved.
+
+#pragma once
+
+#include "AudioPropertiesParserNameMatch.h"
+
+#include "AudioPropertiesParserBindings.generated.h"
+
+class UAudioPropertiesBindings;
+struct FAudioPropertiesSheet;
+
+/*
+*	UAudioPropertiesParserBindings
+* 
+*	Parses values from a Property Sheet to a target UObject 
+*	following the name mappings defined in the BindingsAsset.
+* 
+*	Properties need to be of compatible types for values to be parsed.
+* 
+*/
+UCLASS(MinimalAPI)
+class UAudioPropertiesParserBindings : public UAudioPropertiesParserNameMatch
+{
+	GENERATED_BODY()
+
+public:
+	//The bindings that will be used to parse properties from the Sheet to the target asset
+	UPROPERTY(EditAnywhere, Category = "AudioProperties")
+	TObjectPtr<UAudioPropertiesBindings> BindingsAsset;
+
+	virtual bool GenerateParsingData(TObjectPtr<const UObject> TargetObject, const FAudioPropertiesSheet& SourceSheet, UE::AudioGameplay::FAudioPropertiesParsingData& InOutData) const override;
+	virtual bool ParseProperties(TObjectPtr<UObject> TargetObject, const FAudioPropertiesSheet& PropertiesToParse) const override;
+
+protected:
+
+#if WITH_EDITOR
+	virtual bool GenerateParsingDataFromTargetClass(UClass* TargetClass, UE::AudioGameplay::FAudioPropertiesParsingData& InOutData) const override;
+#endif
+};
\ No newline at end of file
diff --git a/Engine/Plugins/Experimental/AudioProperties/Source/AudioProperties/Public/AudioPropertiesParserNameMatch.h b/Engine/Plugins/Experimental/AudioProperties/Source/AudioProperties/Public/AudioPropertiesParserNameMatch.h
new file mode 100644
--- /dev/null
+++ b/Engine/Plugins/Experimental/AudioProperties/Source/AudioProperties/Public/AudioPropertiesParserNameMatch.h
@@ -0,0 +1,42 @@
+// Copyright Epic Games, Inc. All Rights Reserved.
+
+#pragma once
+
+#include "AudioPropertiesParserBase.h"
+#include "Templates/SubclassOf.h"
+
+#include "AudioPropertiesParserNameMatch.generated.h"
+
+/*
+*	UAudioPropertiesParserNameMatch
+* 
+*	Parses values from a Property Sheet to a target UObject 
+*	of compatible type properties with the same name
+* 
+*/
+UCLASS(MinimalAPI)
+class UAudioPropertiesParserNameMatch : public UAudioPropertiesParserBase
+{
+	GENERATED_BODY()
+
+public:
+	virtual bool GenerateParsingData(TObjectPtr<const UObject> TargetObject, const FAudioPropertiesSheet& SourceSheet, UE::AudioGameplay::FAudioPropertiesParsingData& InOutData) const override;
+	virtual bool ParseProperties(TObjectPtr<UObject> TargetObject, const FAudioPropertiesSheet& PropertiesToParse) const override;
+
+	UPROPERTY(EditAnywhere, Category = "Instanced Objects Parsing")
+	TArray<FInstancedObjectParsingInstructions> InstancedObjectsRules;
+	
+#if WITH_EDITOR
+	bool ValidatePropertiesOnAssetTree(const UAudioPropertiesSheetAsset& InSheetAsset) const override;
+	virtual void FitPropertiesOnAssetTree(UAudioPropertiesSheetAsset& InSheetAsset) const override;
+
+	UPROPERTY(EditAnywhere, Category = "Validation")
+	TSubclassOf<UObject> ValidationClass;
+#endif
+
+protected:
+
+#if WITH_EDITOR
+	virtual bool GenerateParsingDataFromTargetClass(UClass* TargetClass, UE::AudioGameplay::FAudioPropertiesParsingData& InOutData) const;
+#endif
+};
\ No newline at end of file
diff --git a/Engine/Plugins/Experimental/AudioProperties/Source/AudioProperties/Public/AudioPropertiesParsingData.h b/Engine/Plugins/Experimental/AudioProperties/Source/AudioProperties/Public/AudioPropertiesParsingData.h
new file mode 100644
--- /dev/null
+++ b/Engine/Plugins/Experimental/AudioProperties/Source/AudioProperties/Public/AudioPropertiesParsingData.h
@@ -0,0 +1,27 @@
+// Copyright Epic Games, Inc. All Rights Reserved.
+
+#pragma once
+
+#include "Containers/Map.h"
+#include "UObject/NameTypes.h"
+
+class UClass;
+struct FAudioPropertiesSheet;
+
+namespace UE
+{
+	namespace AudioGameplay
+	{
+		using FAudioPropertiesBindings = TMap<FName, FName>;
+
+		struct FAudioPropertiesParsingData
+		{
+			const FAudioPropertiesSheet* SourcePropertySheet = nullptr;
+			UClass* TargetClass = nullptr;
+			FAudioPropertiesBindings SourceToTargetPropertyBindings;
+			bool bCreateNewObjects = false;
+			TArray<FAudioPropertiesParsingData> InstancedObjectsParsingData;
+		};
+	}
+}
+
diff --git a/Engine/Plugins/Experimental/AudioProperties/Source/AudioProperties/Public/AudioPropertiesSheet.h b/Engine/Plugins/Experimental/AudioProperties/Source/AudioProperties/Public/AudioPropertiesSheet.h
new file mode 100644
--- /dev/null
+++ b/Engine/Plugins/Experimental/AudioProperties/Source/AudioProperties/Public/AudioPropertiesSheet.h
@@ -0,0 +1,113 @@
+// Copyright Epic Games, Inc. All Rights Reserved.
+
+#pragma once
+
+#include "AudioPropertiesSheetAssetBase.h"
+#include "Delegates/Delegate.h"
+#include "StructUtils/PropertyBag.h"
+
+#include "AudioPropertiesSheet.generated.h"
+
+#define UE_API AUDIOPROPERTIES_API
+
+class IAudioPropertiesParser;
+class UAudioPropertiesBindings;
+class UAudioPropertiesParserBase;
+
+namespace AudioPropertiesSheet
+{
+	typedef TPair<const FPropertyBagPropertyDesc*, TObjectPtr<const UAudioPropertiesSheetAsset>> PropertyDescParentPair;
+	typedef TPair<TObjectPtr<const UAudioPropertiesParserBase>, TObjectPtr<const UAudioPropertiesSheetAsset>> ParserParentPair;
+
+	#if WITH_EDITOR
+	enum class PostEditChangeType : uint8
+	{
+		Other,
+		ParentChange, 
+		PropertyBag,
+		PropertiesParser
+	};
+#endif
+}
+
+#if WITH_EDITOR
+DECLARE_MULTICAST_DELEGATE_OneParam(FOnAudioPropertySheetPostEditChange, AudioPropertiesSheet::PostEditChangeType);
+DECLARE_MULTICAST_DELEGATE_TwoParams(FOnAudioPropertyOverride, const FProperty& /*TargetProperty*/, bool /*bIsOverridden*/);
+#endif
+
+USTRUCT()
+struct FAudioPropertiesSheet
+{
+	GENERATED_BODY()
+
+	UPROPERTY(EditAnywhere, Category = "Default")
+	TObjectPtr<const UAudioPropertiesSheetAsset> Parent;
+
+	UPROPERTY(EditAnywhere, Category = "Default")
+	FInstancedPropertyBag Properties;
+
+	UE_API bool UpdatePropertyOverride(const FProperty& InProperty, bool bMarkAsOverridden);
+
+	UE_API bool IsPropertyOverridden(const FProperty& InProperty) const;
+
+	UE_API void ReconcileProperties();
+
+	UE_API void OnPreSave(const FObjectPreSaveContext& SaveContext);
+
+	UE_API AudioPropertiesSheet::PropertyDescParentPair FindClosestParentWithProperty(const FProperty& InProperty) const;
+
+#if WITH_EDITOR
+	FOnAudioPropertyOverride OnAudioPropertyOverrideChange;
+#endif	
+private:
+	UE_API void FlattenSheet();
+
+};
+
+UCLASS(MinimalAPI, BlueprintType)
+class UAudioPropertiesSheetAsset : public UAudioPropertiesSheetAssetBase
+{
+	GENERATED_BODY()
+
+public:
+	UE_API UAudioPropertiesSheetAsset(const FObjectInitializer& ObjectInitializer);
+
+	UPROPERTY(EditAnywhere, Category = "Default")
+	FAudioPropertiesSheet PropertiesSheet;
+
+		//The properties parser determines how the properties on this sheet are applied to target UObjects
+	//NOTE: A null parser will mean no properties are applied - a child sheet will however inherit the parent parser automatically
+	UPROPERTY(EditAnywhere, Instanced, Category = "Parser")
+	TObjectPtr<UAudioPropertiesParserBase> PropertiesParser;
+
+#if WITH_EDITOR
+	//Calling this method will apply the parser validation rules to the sheet
+	//NOTE will dirty this asset and all the parents where property changes are applied
+	UFUNCTION(CallInEditor, meta=(DisplayName="Fit Properties For Validation"))
+	UE_API void FitPropertiesForValidation();
+
+	UE_API virtual bool CopyToObjectProperties(UObject* TargetObject) const override;
+
+	UE_API virtual FDelegateHandle BindPropertiesCopyToSheetChanges(UObject* TargetObject) override;
+	UE_API virtual void UnbindCopyFromPropertySheetChanges(UObject* ObjectToUnbind) override;
+
+	UE_API void PreEditChange(FProperty* PropertyAboutToChange) override;
+	UE_API void PostEditChangeProperty(struct FPropertyChangedEvent& PropertyChangedEvent) override;
+	FOnAudioPropertySheetPostEditChange OnPostEditChange;
+
+	UE_API void PostLoad();
+
+	UE_API EDataValidationResult IsDataValid(class FDataValidationContext& Context) const override;
+
+	UE_API AudioPropertiesSheet::ParserParentPair FindClosestParserInInheritanceTree() const;
+
+	UE_API void GetTargetedPropertyNames(TObjectPtr<const UObject> TargetObject, TArray<FName>& OutProperties) const;
+
+	UE_API void OnParentPostEditChange(const AudioPropertiesSheet::PostEditChangeType& ChangeType);
+#endif 
+
+	UE_API void PreSave(FObjectPreSaveContext SaveContext) override;
+
+};
+
+#undef UE_API
diff --git a/Engine/Plugins/Experimental/AudioProperties/Source/AudioProperties/Public/AudioPropertiesUtils.h b/Engine/Plugins/Experimental/AudioProperties/Source/AudioProperties/Public/AudioPropertiesUtils.h
new file mode 100644
--- /dev/null
+++ b/Engine/Plugins/Experimental/AudioProperties/Source/AudioProperties/Public/AudioPropertiesUtils.h
@@ -0,0 +1,22 @@
+// Copyright Epic Games, Inc. All Rights Reserved.
+
+#pragma once 
+
+#include "AssetRegistry/AssetData.h"
+
+class UAudioPropertiesSheetAsset;
+struct FPropertyBagPropertyDesc;
+struct FAudioPropertiesSheet;
+
+namespace AudioPropertiesUtils
+{
+	static const FLazyName AudioPropertiesOverrideTagName = "AudioPropertyOverride";
+
+	bool AUDIOPROPERTIES_API IsSheetInInheritanceTree(const UAudioPropertiesSheetAsset& InLeafSheet, const TObjectPtr<const UAudioPropertiesSheetAsset> InSheetToFind);
+	void GetAudioPropertiesSheetAssetReferencers(const UAudioPropertiesSheetAsset& InAsset, TArray<FAssetData>& OutReferencersData, const UClass* ClassFilter = nullptr);
+	void VisitAudioPropertiesSheetAssetReferencers(const UAudioPropertiesSheetAsset& InAsset, TFunctionRef<void(FAssetData&)> OnAssetVisited, const UClass* ClassFilter = nullptr);
+	void AUDIOPROPERTIES_API VisitLeafMostProperties(const UAudioPropertiesSheetAsset& InSheetAsset, TFunction<void(const FPropertyBagPropertyDesc& /*Property Desc */, const FAudioPropertiesSheet& /*Owning Sheet*/)> OnPropertyVisited);
+	void VisitLeafMostProperties(const FAudioPropertiesSheet& InSheetAsset, TFunction<void(const FPropertyBagPropertyDesc& /*Property Desc */, const FAudioPropertiesSheet& /*Owning Sheet*/)> OnPropertyVisited);
+
+
+}
\ No newline at end of file
diff --git a/Engine/Plugins/Experimental/AudioProperties/Source/AudioPropertiesEditor/AudioPropertiesEditor.build.cs b/Engine/Plugins/Experimental/AudioProperties/Source/AudioPropertiesEditor/AudioPropertiesEditor.build.cs
new file mode 100644
--- /dev/null
+++ b/Engine/Plugins/Experimental/AudioProperties/Source/AudioPropertiesEditor/AudioPropertiesEditor.build.cs
@@ -0,0 +1,43 @@
+// Copyright Epic Games, Inc. All Rights Reserved.
+
+using UnrealBuildTool;
+using System.IO;
+
+namespace UnrealBuildTool.Rules
+{
+	public class AudioPropertiesEditor : ModuleRules
+	{
+		public AudioPropertiesEditor(ReadOnlyTargetRules Target) : base(Target)
+		{
+			PCHUsage = PCHUsageMode.UseExplicitOrSharedPCHs;
+
+			PublicDependencyModuleNames.AddRange(
+				new string[] 
+				{
+					"AudioExtensions",
+					"AudioProperties",
+					"AssetDefinition",
+					"ContentBrowser",
+					"Core",
+					"CoreUObject",
+					"Engine",
+					"InputCore",
+					"PropertyEditor",
+					"Slate",
+					"SlateCore",
+					"ToolMenus",
+					"StructUtilsEditor",
+					"UnrealEd"
+				}
+			);
+
+			PrivateDependencyModuleNames.AddRange(
+				new string[] 
+				{
+					"AssetDefinition",
+					"AudioEditor"
+				}
+			);
+		}
+	}
+}
\ No newline at end of file
diff --git a/Engine/Plugins/Experimental/AudioProperties/Source/AudioPropertiesEditor/Private/AssetDefinition_AudioPropertiesBindings.cpp b/Engine/Plugins/Experimental/AudioProperties/Source/AudioPropertiesEditor/Private/AssetDefinition_AudioPropertiesBindings.cpp
new file mode 100644
--- /dev/null
+++ b/Engine/Plugins/Experimental/AudioProperties/Source/AudioPropertiesEditor/Private/AssetDefinition_AudioPropertiesBindings.cpp
@@ -0,0 +1,9 @@
+﻿// Copyright Epic Games, Inc. All Rights Reserved.
+
+#include "AssetDefinition_AudioPropertiesBindings.h"
+
+TConstArrayView<FAssetCategoryPath> UAssetDefinition_AudioPropertiesBindings::GetAssetCategories() const
+{
+	static const auto Categories = { FAssetCategoryPath(EAssetCategoryPaths::Audio, NSLOCTEXT("AssetTypeActions", "AssetAudioPropertiesSubMenu", "Audio Properties")) };
+	return Categories;
+}
\ No newline at end of file
diff --git a/Engine/Plugins/Experimental/AudioProperties/Source/AudioPropertiesEditor/Private/AssetDefinition_AudioPropertiesBindings.h b/Engine/Plugins/Experimental/AudioProperties/Source/AudioPropertiesEditor/Private/AssetDefinition_AudioPropertiesBindings.h
new file mode 100644
--- /dev/null
+++ b/Engine/Plugins/Experimental/AudioProperties/Source/AudioPropertiesEditor/Private/AssetDefinition_AudioPropertiesBindings.h
@@ -0,0 +1,23 @@
+﻿// Copyright Epic Games, Inc. All Rights Reserved.
+
+#pragma once
+
+#include "AssetTypeActions/AssetDefinitionDefault_AudioDiffable.h"
+#include "AudioPropertiesBindings.h"
+#include "AudioPropertiesEditorModule.h"
+
+#include "AssetDefinition_AudioPropertiesBindings.generated.h"
+
+UCLASS()
+class UAssetDefinition_AudioPropertiesBindings : public UAssetDefinitionDefault_AudioDiffable
+{
+	GENERATED_BODY()
+
+public:
+	//~ Begin UAssetDefinition
+	virtual FText GetAssetDisplayName() const override { return NSLOCTEXT("AssetTypeActions", "AssetTypeActions_AudioPropertiesBindings", "Audio Properties Bindings"); }
+	virtual FLinearColor GetAssetColor() const override { return FLinearColor(AudioPropertiesEditorModule::AssetColor); }
+	virtual TSoftClassPtr<UObject> GetAssetClass() const override { return UAudioPropertiesBindings::StaticClass(); }
+	virtual TConstArrayView<FAssetCategoryPath> GetAssetCategories() const override;
+	//~ End UAssetDefinition
+};
diff --git a/Engine/Plugins/Experimental/AudioProperties/Source/AudioPropertiesEditor/Private/AssetDefinition_AudioPropertiesSheetAsset.cpp b/Engine/Plugins/Experimental/AudioProperties/Source/AudioPropertiesEditor/Private/AssetDefinition_AudioPropertiesSheetAsset.cpp
new file mode 100644
--- /dev/null
+++ b/Engine/Plugins/Experimental/AudioProperties/Source/AudioPropertiesEditor/Private/AssetDefinition_AudioPropertiesSheetAsset.cpp
@@ -0,0 +1,58 @@
+﻿// Copyright Epic Games, Inc. All Rights Reserved.
+
+#include "AssetDefinition_AudioPropertiesSheetAsset.h"
+
+#include "ContentBrowserMenuContexts.h"
+#include "Toolkits/SimpleAssetEditor.h"
+#include "ToolMenus.h"
+#include "ToolMenuSection.h"
+
+#define LOCTEXT_NAMESPACE "AudioPropertiesEditor"
+
+TConstArrayView<FAssetCategoryPath> UAssetDefinition_AudioPropertiesSheetAsset::GetAssetCategories() const
+{
+	static const auto Categories = { FAssetCategoryPath(EAssetCategoryPaths::Audio, NSLOCTEXT("AssetTypeActions", "AssetAudioPropertiesSubMenu", "Audio Properties")) };
+	return Categories;
+}
+
+EAssetCommandResult UAssetDefinition_AudioPropertiesSheetAsset::OpenAssets(const FAssetOpenArgs& OpenArgs) const
+{
+	for (UAudioPropertiesSheetAsset* PropertySheetAsset : OpenArgs.LoadObjects<UAudioPropertiesSheetAsset>())
+	{
+		FSimpleAssetEditor::CreateEditor(EToolkitMode::Standalone, OpenArgs.ToolkitHost, PropertySheetAsset);
+	}
+	
+	return EAssetCommandResult::Handled;
+}
+
+FToolMenuSection* UAssetDefinition_AudioPropertiesSheetAsset::FindContextMenuSection(FName SectionName) const
+{
+	TArray<FToolMenuSection*> Sections = RebuildContextMenuSections();
+	
+	for (FToolMenuSection* Section : Sections)
+	{
+		if (Section->Name == SectionName)
+		{
+			return Section;
+		}
+	}
+
+	return nullptr;
+}
+
+TArray<FToolMenuSection*> UAssetDefinition_AudioPropertiesSheetAsset::RebuildContextMenuSections() const
+{
+	TArray<FToolMenuSection*> Sections;
+	
+	UToolMenu* Menu = UE::ContentBrowser::ExtendToolMenu_AssetContextMenu(GetAssetClass());
+	check(Menu);
+
+	FToolMenuSection& SoundSection = Menu->FindOrAddSection("Builder");
+	SoundSection.Label = LOCTEXT("AudioPropertiesSheetAssetActions_Label", "Builder");
+
+	Sections.Add(&SoundSection);
+
+	return Sections;
+}
+
+#undef LOCTEXT_NAMESPACE
\ No newline at end of file
diff --git a/Engine/Plugins/Experimental/AudioProperties/Source/AudioPropertiesEditor/Private/AudioPropertiesBindingsFactory.cpp b/Engine/Plugins/Experimental/AudioProperties/Source/AudioPropertiesEditor/Private/AudioPropertiesBindingsFactory.cpp
new file mode 100644
--- /dev/null
+++ b/Engine/Plugins/Experimental/AudioProperties/Source/AudioPropertiesEditor/Private/AudioPropertiesBindingsFactory.cpp
@@ -0,0 +1,22 @@
+// Copyright Epic Games, Inc. All Rights Reserved.
+
+#include "AudioPropertiesBindingsFactory.h"
+
+#include "AudioPropertiesBindings.h"
+
+#include UE_INLINE_GENERATED_CPP_BY_NAME(AudioPropertiesBindingsFactory)
+
+UAudioPropertiesBindingsFactory::UAudioPropertiesBindingsFactory(const FObjectInitializer& ObjectInitializer)
+	: Super(ObjectInitializer)
+{
+	SupportedClass = UAudioPropertiesBindings::StaticClass();
+	bCreateNew = true;
+	bEditorImport = false;
+	bEditAfterNew = true;
+}
+
+
+UObject* UAudioPropertiesBindingsFactory::FactoryCreateNew(UClass* Class, UObject* InParent, FName Name, EObjectFlags Flags, UObject* Context, FFeedbackContext* Warn)
+{
+	return NewObject<UAudioPropertiesBindings>(InParent, Name, Flags);
+}
diff --git a/Engine/Plugins/Experimental/AudioProperties/Source/AudioPropertiesEditor/Private/AudioPropertiesDetailsInjector.cpp b/Engine/Plugins/Experimental/AudioProperties/Source/AudioPropertiesEditor/Private/AudioPropertiesDetailsInjector.cpp
new file mode 100644
--- /dev/null
+++ b/Engine/Plugins/Experimental/AudioProperties/Source/AudioPropertiesEditor/Private/AudioPropertiesDetailsInjector.cpp
@@ -0,0 +1,240 @@
+// Copyright Epic Games, Inc. All Rights Reserved.
+
+#include "AudioPropertiesDetailsInjector.h"
+
+#include "AudioPropertiesDetailsInjectorUtils.h"
+#include "AudioPropertiesSheet.h"
+#include "CustomNodeBuilders/AudioPropertiesDetailArrayBuilder.h"
+#include "CustomNodeBuilders/AudioPropertiesDetailNestedNodesBuilder.h"
+#include "CustomNodeBuilders/AudioPropertiesDetailSetBuilder.h"
+#include "DetailLayoutBuilder.h"
+#include "DetailWidgetRow.h"
+#include "IPropertyUtilities.h"
+#include "PropertyCustomizationHelpers.h"
+#include "UObject/Object.h"
+
+#define LOCTEXT_NAMESPACE "AudioPropertiesDetailsInjection"
+
+void FAudioPropertiesDetailsInjector::CustomizeInjectedPropertiesDetails(IDetailLayoutBuilder& DetailBuilder, TSharedRef<IPropertyHandle> PropertySheetPropertyHandle)
+{
+	if (!PropertySheetPropertyHandle->IsValidHandle())
+	{
+		return;
+	}
+
+	TArray<TWeakObjectPtr<UObject>> ObjectsBeingCustomized;
+	DetailBuilder.GetObjectsBeingCustomized(ObjectsBeingCustomized);
+
+	CreateOverriddenPropertiesWidgets(DetailBuilder, ObjectsBeingCustomized, PropertySheetPropertyHandle);
+}
+
+void FAudioPropertiesDetailsInjector::CreateOverriddenPropertiesWidgets(IDetailLayoutBuilder& DetailBuilder, TArray<TWeakObjectPtr<UObject>> ObjectsBeingCustomized, TSharedRef<IPropertyHandle> PropertySheetPropertyHandle)
+{
+	for (TWeakObjectPtr<UObject>& WeakObjectBeingCustomized : ObjectsBeingCustomized)
+	{
+		TObjectPtr<UObject> ObjectBeingCustomized = WeakObjectBeingCustomized.Get();
+
+		if (!ObjectBeingCustomized || !ObjectBeingCustomized->IsValidLowLevel())
+		{
+			continue;
+		}
+
+		const UAudioPropertiesSheetAsset* SheetAssetPtr = Cast<UAudioPropertiesSheetAsset>(AudioPropertiesDetailsInjectorUtils::GetUObjectFromPtrPropertyHandle(PropertySheetPropertyHandle, ObjectBeingCustomized));
+
+		if (!SheetAssetPtr || !SheetAssetPtr->IsValidLowLevel())
+		{
+			continue;
+		}
+
+		TArray<FName> TargetedPropertyNames;
+
+		SheetAssetPtr->GetTargetedPropertyNames(ObjectBeingCustomized, TargetedPropertyNames);
+
+
+		for (const FName& TargetPropertyName : TargetedPropertyNames)
+		{
+			TSharedRef<IPropertyHandle> OverriddenPropertyHandle = DetailBuilder.GetProperty(TargetPropertyName);
+
+			if (!OverriddenPropertyHandle->IsValidHandle())
+			{
+				//If GetProperty fails use the class path
+
+				OverriddenPropertyHandle = DetailBuilder.GetProperty(TargetPropertyName, ObjectBeingCustomized->GetClass());
+				if (!OverriddenPropertyHandle->IsValidHandle())
+				{
+					continue;
+				}
+			}
+			
+			IDetailPropertyRow* OverriddenPropertyRow = DetailBuilder.EditDefaultProperty(OverriddenPropertyHandle);
+			check(OverriddenPropertyRow)
+
+			if (OverriddenPropertyHandle->AsArray().IsValid())
+			{
+				const FName DetailCategory = OverriddenPropertyHandle->GetDefaultCategoryName();
+				TSharedRef<FAudioPropertiesDetailArrayBuilder> OverriddenArrayDetails = MakeShared<FAudioPropertiesDetailArrayBuilder>(OverriddenPropertyHandle, PropertySheetPropertyHandle, ObjectBeingCustomized);
+				DetailBuilder.EditCategory(DetailCategory).AddCustomBuilder(OverriddenArrayDetails);
+			}
+			else if (OverriddenPropertyHandle->AsSet().IsValid())
+			{
+				const FName DetailCategory = OverriddenPropertyHandle->GetDefaultCategoryName();
+				TSharedRef<FAudioPropertiesDetailSetBuilder> OverriddenSetDetails = MakeShared<FAudioPropertiesDetailSetBuilder>(OverriddenPropertyHandle, PropertySheetPropertyHandle, ObjectBeingCustomized);
+				DetailBuilder.EditCategory(DetailCategory).AddCustomBuilder(OverriddenSetDetails);
+			}
+			else
+			{
+				uint32 NumChildren = 0;
+				OverriddenPropertyHandle->GetNumChildren(NumChildren);
+
+				if (NumChildren > 0)
+				{
+					const FName DetailCategory = OverriddenPropertyHandle->GetDefaultCategoryName();
+					TSharedRef<FAudioPropertiesDetailNestedNodesBuilder> OverriddenDetails = MakeShared<FAudioPropertiesDetailNestedNodesBuilder>(OverriddenPropertyHandle, PropertySheetPropertyHandle, ObjectBeingCustomized);
+					DetailBuilder.EditCategory(DetailCategory).AddCustomBuilder(OverriddenDetails);
+				}
+				else
+				{
+					TSharedPtr<SWidget> NameWidget;
+					TSharedPtr<SWidget> ValueWidget;
+					FDetailWidgetRow PropertyRow;
+					OverriddenPropertyRow->GetDefaultWidgets(NameWidget, ValueWidget, PropertyRow);
+
+					OverriddenPropertyRow->CustomWidget(true)
+						.NameContent()
+						[
+							AudioPropertiesDetailsInjectorUtils::CreateOverriddenPropertyNameWidget(OverriddenPropertyHandle, PropertySheetPropertyHandle, ObjectsBeingCustomized)
+						]
+						.ValueContent()
+						[
+							SNew(SHorizontalBox)
+							.IsEnabled_Lambda([OverriddenPropertyHandle, ObjectBeingCustomized]()
+							{
+								return !AudioPropertiesDetailsInjectorUtils::ShouldParseProperty(ObjectBeingCustomized, OverriddenPropertyHandle);
+							})
+							+ SHorizontalBox::Slot()
+							[
+								ValueWidget.ToSharedRef()
+							]
+						];
+
+					TAttribute<bool> ResetToDefaultEnabled = TAttribute<bool>::CreateLambda([OverriddenPropertyHandle, ObjectBeingCustomized](){ return !AudioPropertiesDetailsInjectorUtils::ShouldParseProperty(ObjectBeingCustomized, OverriddenPropertyHandle);;});
+					OverriddenPropertyRow->OverrideResetToDefault(FResetToDefaultOverride::Create(ResetToDefaultEnabled));
+				}
+			}
+		}
+	}
+}
+
+void FAudioPropertiesDetailsInjector::BindDetailCustomizationToPropertySheetChanges(IDetailLayoutBuilder& DetailBuilder, TSharedRef<IPropertyHandle> PropertySheetPropertyHandle)
+{
+	if (!PropertySheetPropertyHandle->IsValidHandle())
+	{
+		return;
+	}
+
+	TSharedRef<IPropertyUtilities> PropertyUtils = DetailBuilder.GetPropertyUtilities();
+
+	TArray<TWeakObjectPtr<UObject>> ObjectsBeingCustomized;
+	DetailBuilder.GetObjectsBeingCustomized(ObjectsBeingCustomized);
+
+	BindDetailsRefreshToPropertySheetSwaps(DetailBuilder, ObjectsBeingCustomized, PropertySheetPropertyHandle);
+	BindDetailsRefreshToPropertySheetChanges(DetailBuilder, ObjectsBeingCustomized, PropertySheetPropertyHandle);
+
+}
+
+void FAudioPropertiesDetailsInjector::BindDetailsRefreshToPropertySheetSwaps(const IDetailLayoutBuilder& DetailBuilder, TArray<TWeakObjectPtr<UObject>>& ObjectsBeingCustomized, TSharedRef<IPropertyHandle> PropertySheetPropertyHandle)
+{
+	TSharedRef<IPropertyUtilities> PropertyUtils = DetailBuilder.GetPropertyUtilities();
+
+	auto OnMatrixPropertySheetChanged = [PropertyUtils, PropertySheetPropertyHandle, ObjectsBeingCustomized](UObject* ObjectBeingModified, FPropertyChangedEvent& PropertyChangedEvent) {
+		
+		if(!PropertyChangedEvent.Property)
+		{
+			return;
+		}
+		
+		if (!ObjectsBeingCustomized.Contains(ObjectBeingModified))
+		{
+			return;
+		}
+		
+		if (PropertyChangedEvent.Property == PropertySheetPropertyHandle->GetProperty())
+		{
+			PropertyUtils->ForceRefresh();
+		}
+	};
+
+	// Preferring this over IPropertyHandle::SetOnPropertyValueChanged because it supports multi edit and editor transactions
+	FDelegateHandle Handle = FCoreUObjectDelegates::OnObjectPropertyChanged.AddLambda(OnMatrixPropertySheetChanged);
+	OnTargetObjectPropertyChangeDelegateHandles.Add(Handle);
+}
+
+void FAudioPropertiesDetailsInjector::BindDetailsRefreshToPropertySheetChanges(IDetailLayoutBuilder& DetailBuilder, TArray<TWeakObjectPtr<UObject>>& ObjectsBeingCustomized, TSharedRef<IPropertyHandle> PropertySheetPropertyHandle)
+{
+	for (TWeakObjectPtr<UObject>& WeakObjectBeingCustomized : ObjectsBeingCustomized)
+	{
+		TObjectPtr<UObject> ObjectBeingCustomized = WeakObjectBeingCustomized.Get();
+
+		if (!ObjectBeingCustomized || !ObjectBeingCustomized->IsValidLowLevel())
+		{
+			continue;
+		}
+
+		const UAudioPropertiesSheetAsset* SheetAssetPtr = Cast<UAudioPropertiesSheetAsset>(AudioPropertiesDetailsInjectorUtils::GetUObjectFromPtrPropertyHandle(PropertySheetPropertyHandle, ObjectBeingCustomized));
+
+		if (!SheetAssetPtr || !SheetAssetPtr->IsValidLowLevel())
+		{
+			continue;
+		}
+
+		UAudioPropertiesSheetAsset* NonConstPtrForBinding = const_cast<UAudioPropertiesSheetAsset*>(SheetAssetPtr);
+
+		auto OnPropertiesSheetPostEditChange = [DetailsView = DetailBuilder.GetDetailsViewSharedPtr().ToWeakPtr()](AudioPropertiesSheet::PostEditChangeType InChangeType)
+		{
+			switch (InChangeType)
+			{
+			case AudioPropertiesSheet::PostEditChangeType::PropertyBag:
+			case AudioPropertiesSheet::PostEditChangeType::ParentChange:
+			case AudioPropertiesSheet::PostEditChangeType::PropertiesParser:
+				if (TSharedPtr<IDetailsView> DetailsViewPinned = DetailsView.Pin())
+				{
+					if (DetailsViewPinned->IsParentValid())
+					{
+						DetailsViewPinned->ForceRefresh();
+					}
+				}
+				break;
+			case AudioPropertiesSheet::PostEditChangeType::Other:
+				break;
+			default:
+				break;
+			}
+		};
+
+		FDelegateHandle PostEditDelegateBinding = NonConstPtrForBinding->OnPostEditChange.AddLambda(OnPropertiesSheetPostEditChange);
+
+		InjectedPropertiesBindings.Add(NonConstPtrForBinding, PostEditDelegateBinding);
+	}
+}
+
+void FAudioPropertiesDetailsInjector::UnbindFromPropertySheetChanges()
+{
+	for(const FDelegateHandle& Handle : OnTargetObjectPropertyChangeDelegateHandles)
+	{
+		FCoreUObjectDelegates::OnObjectPropertyChanged.Remove(Handle);
+	}
+
+	OnTargetObjectPropertyChangeDelegateHandles.Empty();
+	
+	for (TPair<FObjectKey, FDelegateHandle>& PropertyBinding : InjectedPropertiesBindings)
+	{
+		if (UObject* BoundObj = PropertyBinding.Key.ResolveObjectPtr())
+		{
+			if (UAudioPropertiesSheetAsset* SheetAssetPtr = Cast<UAudioPropertiesSheetAsset>(BoundObj))
+			{
+				SheetAssetPtr->OnPostEditChange.Remove(PropertyBinding.Value);
+			}
+		}
+	}
+}
+#undef LOCTEXT_NAMESPACE 
\ No newline at end of file
diff --git a/Engine/Plugins/Experimental/AudioProperties/Source/AudioPropertiesEditor/Private/AudioPropertiesDetailsInjectorUtils.h b/Engine/Plugins/Experimental/AudioProperties/Source/AudioPropertiesEditor/Private/AudioPropertiesDetailsInjectorUtils.h
new file mode 100644
--- /dev/null
+++ b/Engine/Plugins/Experimental/AudioProperties/Source/AudioPropertiesEditor/Private/AudioPropertiesDetailsInjectorUtils.h
@@ -0,0 +1,262 @@
+﻿// Copyright Epic Games, Inc. All Rights Reserved.
+
+#pragma once
+
+#include "AudioPropertiesSheet.h"
+#include "DetailLayoutBuilder.h"
+#include "ScopedTransaction.h"
+#include "Widgets/Input/SCheckBox.h"
+#include "Widgets/Input/SComboButton.h"
+#include "Widgets/SToolTip.h"
+#include "Widgets/Text/STextBlock.h"
+
+#define LOCTEXT_NAMESPACE "AudioPropertiesDetailsInjectorUtils"
+
+namespace AudioPropertiesDetailsInjectorUtils
+{
+	static const FLazyName EditConditionPropertyMetadata = "EditCondition";
+	static const FLazyName InlineEditConditionTogglePropertyMetadata = "InlineEditConditionToggle";
+	
+	inline const UObject* GetUObjectFromPtrPropertyHandle(const TSharedPtr<IPropertyHandle> InPropertyHandle, const UObject* PropertyOwner)
+	{
+		check(InPropertyHandle)
+		check(PropertyOwner)
+		
+		const FProperty* TargetProperty = InPropertyHandle->GetProperty();
+		
+		if (!ensure(TargetProperty))
+		{
+			return nullptr;
+		}
+
+		const FName& PropertyName = TargetProperty->GetFName();
+		FProperty* PropertyPtr = PropertyOwner->GetClass()->FindPropertyByName(PropertyName);
+
+		if (PropertyPtr)
+		{
+			if (FObjectPropertyBase* ObjectProperty = CastField<FObjectPropertyBase>(PropertyPtr))
+			{
+				return ObjectProperty->GetObjectPropertyValue(ObjectProperty->ContainerPtrToValuePtr<void>(PropertyOwner));
+			}
+		}
+
+		return nullptr;
+	}
+
+	inline bool ShouldParseProperty(const UObject* PropertyOwner, const TSharedPtr<IPropertyHandle> InPropertyHandle)
+	{
+		if (!ensureMsgf(PropertyOwner, TEXT("Object should not be null")))
+		{
+			return false;
+		}
+
+		const FProperty* TargetProperty = InPropertyHandle->GetProperty();
+		
+		if (!TargetProperty)
+		{
+			return false;
+		}
+
+		const IAudioPropertiesSheetAssetUserInterface* PropertySheetAssetUser = Cast<IAudioPropertiesSheetAssetUserInterface>(PropertyOwner);
+
+		if (!ensureMsgf(PropertySheetAssetUser, TEXT("Implement IAudioPropertiesSheetAssetUserInterface on objects using PropertySheets to support local overrides")))
+		{
+			return false;
+		}
+
+		return PropertySheetAssetUser->ShouldParseProperty(*TargetProperty);
+	}
+
+	inline void HandleLocalPropertyOverrideChange(UObject* PropertyOwner, const bool bIsPropertyOverriden, const TSharedPtr<IPropertyHandle> TargetPropertyHandle, const TSharedPtr<IPropertyHandle> PropertySheetHandle)
+	{
+		if (!ensureMsgf(PropertyOwner, TEXT("Object should not be null")))
+		{
+			return;
+		}
+
+		const FProperty* TargetProperty = TargetPropertyHandle->GetProperty();
+		
+		if (!TargetProperty)
+		{
+			return;
+		}
+
+		IAudioPropertiesSheetAssetUserInterface* PropertySheetAssetUser = Cast<IAudioPropertiesSheetAssetUserInterface>(PropertyOwner);
+
+		if (!ensureMsgf(PropertySheetAssetUser, TEXT("Implement IAudioPropertiesSheetAssetUserInterface on objects using PropertySheets to support local overrides")))
+		{
+			return;
+		}
+
+		FScopedTransaction Transaction(FText::Format(LOCTEXT("OnCheckStateChanged", "Changed Local Override for property {0}"), FText::FromName(TargetProperty->GetFName())));
+		PropertyOwner->Modify();
+
+		if (bIsPropertyOverriden)
+		{
+			PropertySheetAssetUser->IgnorePropertyParsing(*TargetProperty);
+		}
+		else
+		{
+			PropertySheetAssetUser->AllowPropertyParsing(*TargetProperty);
+
+			const UAudioPropertiesSheetAsset* SheetAssetPtr = Cast<UAudioPropertiesSheetAsset>(GetUObjectFromPtrPropertyHandle(PropertySheetHandle, PropertyOwner));
+
+			if (!SheetAssetPtr || !SheetAssetPtr->IsValidLowLevel())
+			{
+				return;
+			}
+
+			SheetAssetPtr->CopyToObjectProperties(PropertyOwner);
+		}
+	}
+
+	inline TSharedRef<SWidget> GetSingleEditPropertyNameWidget(TSharedPtr<IPropertyHandle> PropertyHandle, const TSharedPtr<IPropertyHandle> PropertySheetHandle, const TArray<TWeakObjectPtr<UObject>>& ObjectsBeingCustomized)
+	{
+		checkf(ObjectsBeingCustomized.Num() == 1, TEXT("GetMultiEditPropertyNameWidget should be used when editing multiple objects"))
+
+		const UAudioPropertiesSheetAsset* SheetAssetPtr = Cast<UAudioPropertiesSheetAsset>(AudioPropertiesDetailsInjectorUtils::GetUObjectFromPtrPropertyHandle(PropertySheetHandle, ObjectsBeingCustomized[0].Get()));
+		FText PropertiesSheetName;
+
+		if (SheetAssetPtr)
+		{
+			PropertiesSheetName = FText::FromName(SheetAssetPtr->GetFName());
+		}
+
+		FText TooltipFText = FText::Format(LOCTEXT("OverriddenAudioPropertyTooltip", "Property Overridden by {0}"), PropertiesSheetName);
+
+		TSharedRef<SToolTip> NameWidgetTooltip = SNew(SToolTip)
+		[
+			SNew(STextBlock)
+			.Text(TooltipFText)
+		];
+
+
+		TSharedRef<SWidget> PropertyNameWidget =  SNew(SHorizontalBox)
+			+ SHorizontalBox::Slot()
+			.AutoWidth()
+			[
+				SNew(SCheckBox)
+				.IsChecked_Lambda([PropertyHandle, OwnerObject = ObjectsBeingCustomized[0]]()
+				{
+					return AudioPropertiesDetailsInjectorUtils::ShouldParseProperty(OwnerObject.Get(), PropertyHandle) ? ECheckBoxState::Unchecked : ECheckBoxState::Checked;
+				})
+				.OnCheckStateChanged_Lambda([PropertyHandle, PropertySheetHandle, OwnerObject = ObjectsBeingCustomized[0]](ECheckBoxState NewState)
+				{
+					const bool bIsOverridden = NewState == ECheckBoxState::Checked;
+					HandleLocalPropertyOverrideChange(OwnerObject.Get(), bIsOverridden, PropertyHandle, PropertySheetHandle);
+				})
+			]
+			+ SHorizontalBox::Slot()
+			[
+				SNew(STextBlock)
+				.Font(IDetailLayoutBuilder::GetDetailFont())
+				.Text_Lambda([PropertyHandle, OwnerObject = ObjectsBeingCustomized[0]]()
+				{
+					if (AudioPropertiesDetailsInjectorUtils::ShouldParseProperty(OwnerObject.Get(), PropertyHandle))
+					{
+						return FText::Format(LOCTEXT("OverriddenAudioPropertyDisplayName", "[AudioProperty Override] {0}"), PropertyHandle->GetPropertyDisplayName());
+					}
+
+					return FText::Format(LOCTEXT("OverriddenAudioPropertyDisplayName", "[Local] {0}"), PropertyHandle->GetPropertyDisplayName());
+				})
+				.ToolTip(NameWidgetTooltip)
+				.IsEnabled_Lambda([PropertyHandle, OwnerObject = ObjectsBeingCustomized[0]]()
+				{
+					return !AudioPropertiesDetailsInjectorUtils::ShouldParseProperty(OwnerObject.Get(), PropertyHandle);
+				})
+			];
+
+
+		return PropertyNameWidget;
+	}
+
+	inline TSharedRef<SWidget> GetMultiEditPropertyNameWidget(TSharedPtr<IPropertyHandle> PropertyHandle, const TSharedPtr<IPropertyHandle> PropertySheetHandle, const TArray<TWeakObjectPtr<UObject>>& ObjectsBeingCustomized)
+	{
+		checkf(ObjectsBeingCustomized.Num() > 0, TEXT("GetSingleEditPropertyNameWidget should be used when editing a single object"))
+
+		FText TooltipFText = LOCTEXT("OverriddenAudioPropertyTooltip", "Property Overridden by one or more properties sheet. Use drop down to select local overrides");
+
+		//create a drop-down menu that allows to check which assets in the matrix have local overrides for a specific property
+		TSharedRef<SWidget> PropertyNameWidget = SNew(SHorizontalBox)
+				+ SHorizontalBox::Slot()
+				.AutoWidth()
+				.Padding(5.0f)
+				[
+					SNew(SComboButton)
+						.OnGetMenuContent(FOnGetContent::CreateLambda([PropertyHandle, PropertySheetHandle, ObjectsBeingCustomized]()
+						{
+							// Create a vertical box to hold the check-boxes
+							TSharedPtr<SVerticalBox> VerticalBox;
+
+							SAssignNew(VerticalBox, SVerticalBox);
+
+							// Iterate over the objects and create a checkbox for each
+							for (TWeakObjectPtr<UObject> WeakObject : ObjectsBeingCustomized)
+							{
+								if (WeakObject.IsValid())
+								{
+									// Add an entry to the vertical box for each object
+									VerticalBox->AddSlot()
+										.AutoHeight()
+										[
+											SNew(SHorizontalBox)
+											+ SHorizontalBox::Slot()
+											.AutoWidth()
+											.VAlign(VAlign_Center)
+											[
+												SNew(SCheckBox)
+												.IsChecked_Lambda([PropertyHandle, WeakObject]()
+												{
+													return AudioPropertiesDetailsInjectorUtils::ShouldParseProperty(WeakObject.Get(), PropertyHandle) ? ECheckBoxState::Unchecked : ECheckBoxState::Checked;
+												})
+												.OnCheckStateChanged_Lambda([PropertyHandle, PropertySheetHandle, WeakObject](ECheckBoxState NewState)
+												{
+													const bool bIsOverridden = NewState == ECheckBoxState::Checked;
+													HandleLocalPropertyOverrideChange(WeakObject.Get(), bIsOverridden, PropertyHandle, PropertySheetHandle);
+												})
+												]
+												+ SHorizontalBox::Slot()
+												.Padding(5.0f, 0.0f)
+												.VAlign(VAlign_Center)
+												[
+													SNew(STextBlock)
+														.Text(FText::FromString(WeakObject->GetName()))
+												]
+										];
+								}
+							}
+
+							return SNew(SBorder)
+									.BorderImage(FCoreStyle::Get().GetBrush("Menu.Background"))
+									[
+										VerticalBox.ToSharedRef()
+									];
+						}))
+						.ButtonContent()
+							[
+								SNew(STextBlock)
+								.Text(PropertyHandle->GetPropertyDisplayName())
+								.Font(IDetailLayoutBuilder::GetDetailFont())
+							]
+				];
+
+		return PropertyNameWidget;
+	}
+
+
+	inline TSharedRef<SWidget> CreateOverriddenPropertyNameWidget(TSharedPtr<IPropertyHandle> PropertyHandle,  const TSharedPtr<IPropertyHandle> PropertySheetHandle, const TArray<TWeakObjectPtr<UObject>>& ObjectsBeingCustomized)
+	{
+		check(ObjectsBeingCustomized.Num() > 0)
+		const bool bMultiEdit = ObjectsBeingCustomized.Num() > 1;
+
+		if (!bMultiEdit)
+		{
+			return GetSingleEditPropertyNameWidget(PropertyHandle, PropertySheetHandle, ObjectsBeingCustomized);
+		}
+		else
+		{
+			return GetMultiEditPropertyNameWidget(PropertyHandle, PropertySheetHandle, ObjectsBeingCustomized);
+		}
+	}
+}
+#undef LOCTEXT_NAMESPACE 
diff --git a/Engine/Plugins/Experimental/AudioProperties/Source/AudioPropertiesEditor/Private/AudioPropertiesEditorLogCategory.cpp b/Engine/Plugins/Experimental/AudioProperties/Source/AudioPropertiesEditor/Private/AudioPropertiesEditorLogCategory.cpp
new file mode 100644
--- /dev/null
+++ b/Engine/Plugins/Experimental/AudioProperties/Source/AudioPropertiesEditor/Private/AudioPropertiesEditorLogCategory.cpp
@@ -0,0 +1,5 @@
+// Copyright Epic Games, Inc. All Rights Reserved.
+
+#include "AudioPropertiesEditorLogCategory.h"
+
+DEFINE_LOG_CATEGORY(LogAudioPropertiesEditor);
\ No newline at end of file
diff --git a/Engine/Plugins/Experimental/AudioProperties/Source/AudioPropertiesEditor/Private/AudioPropertiesEditorLogCategory.h b/Engine/Plugins/Experimental/AudioProperties/Source/AudioPropertiesEditor/Private/AudioPropertiesEditorLogCategory.h
new file mode 100644
--- /dev/null
+++ b/Engine/Plugins/Experimental/AudioProperties/Source/AudioPropertiesEditor/Private/AudioPropertiesEditorLogCategory.h
@@ -0,0 +1,7 @@
+// Copyright Epic Games, Inc. All Rights Reserved.
+
+#pragma once
+
+#include "Logging/LogMacros.h"
+
+DECLARE_LOG_CATEGORY_EXTERN(LogAudioPropertiesEditor, Log, All);
diff --git a/Engine/Plugins/Experimental/AudioProperties/Source/AudioPropertiesEditor/Private/AudioPropertiesEditorModule.cpp b/Engine/Plugins/Experimental/AudioProperties/Source/AudioPropertiesEditor/Private/AudioPropertiesEditorModule.cpp
new file mode 100644
--- /dev/null
+++ b/Engine/Plugins/Experimental/AudioProperties/Source/AudioPropertiesEditor/Private/AudioPropertiesEditorModule.cpp
@@ -0,0 +1,48 @@
+// Copyright Epic Games, Inc. All Rights Reserved.
+
+#include "AudioPropertiesEditorModule.h"
+
+#include "AssetToolsModule.h"
+#include "AudioPropertiesSheetBuilderInstantiator.h"
+#include "AudioPropertiesDetailsInjector.h"
+#include "AudioPropertiesSheetAssetDetails.h"
+#include "AudioPropertiesSheetDetails.h"
+#include "Features/IModularFeatures.h"
+#include "Modules/ModuleManager.h"
+#include "PropertyEditorModule.h"
+
+namespace AudioPropertiesEditorModule
+{
+	static const FName AssetToolsModuleName = TEXT("AssetTools");
+}
+
+void FAudioPropertiesEditorModule::StartupModule()
+{
+	IAssetTools& AssetTools = FModuleManager::LoadModuleChecked<FAssetToolsModule>(AudioPropertiesEditorModule::AssetToolsModuleName).Get();
+	
+	FPropertyEditorModule& PropertyModule = FModuleManager::LoadModuleChecked<FPropertyEditorModule>("PropertyEditor");
+	PropertyModule.RegisterCustomPropertyTypeLayout("AudioPropertiesSheet", FOnGetPropertyTypeCustomizationInstance::CreateStatic(&FAudioPropertiesSheetDetails::MakeInstance));
+	PropertyModule.NotifyCustomizationModuleChanged();
+
+	PropertyModule.RegisterCustomClassLayout("AudioPropertiesSheetAsset", FOnGetDetailCustomizationInstance::CreateStatic(&FAudioPropertiesSheetAssetDetails::MakeInstance));
+
+	DetailsInjectorBuilder = MakeShared<FAudioPropertiesDetailsInjectorBuilder>();
+
+	IModularFeatures::Get().RegisterModularFeature(AudioPropertiesDetailsInjector::BuilderModularFeatureName, DetailsInjectorBuilder.Get());
+
+	
+	PropertiesSheetBuilder = MakeShared<FAudioPropertiesSheetBuilderInstantiator>();
+	PropertiesSheetBuilder->ExtendContentBrowserSelectionMenu();
+}
+
+void FAudioPropertiesEditorModule::ShutdownModule()
+{
+	FPropertyEditorModule& PropertyModule = FModuleManager::LoadModuleChecked<FPropertyEditorModule>("PropertyEditor");
+	PropertyModule.UnregisterCustomPropertyTypeLayout("AudioPropertiesSheet");
+	PropertyModule.NotifyCustomizationModuleChanged();
+	
+	IModularFeatures::Get().UnregisterModularFeature(AudioPropertiesDetailsInjector::BuilderModularFeatureName, DetailsInjectorBuilder.Get());
+	DetailsInjectorBuilder.Reset();
+}
+
+IMPLEMENT_MODULE(FAudioPropertiesEditorModule, AudioPropertiesEditor);
\ No newline at end of file
diff --git a/Engine/Plugins/Experimental/AudioProperties/Source/AudioPropertiesEditor/Private/AudioPropertiesSheetAssetBuilder.cpp b/Engine/Plugins/Experimental/AudioProperties/Source/AudioPropertiesEditor/Private/AudioPropertiesSheetAssetBuilder.cpp
new file mode 100644
--- /dev/null
+++ b/Engine/Plugins/Experimental/AudioProperties/Source/AudioPropertiesEditor/Private/AudioPropertiesSheetAssetBuilder.cpp
@@ -0,0 +1,139 @@
+// Copyright Epic Games, Inc. All Rights Reserved.
+
+#include "AudioPropertiesSheetAssetBuilder.h"
+
+#include "AudioPropertiesSheet.h"
+#include "AudioPropertiesSheetAssetFactory.h"
+#include "AssetToolsModule.h"
+#include "AssetRegistry/AssetRegistryModule.h"
+#include "AudioPropertiesUtils.h"
+#include "Dialogs/DlgPickAssetPath.h"
+#include "UObject/PropertyAccessUtil.h"
+
+#define LOCTEXT_NAMESPACE "AudioPropertiesSheetBuilderWidget"
+
+void FAudioPropertiesSheetAssetBuilder::BuildPropertySheetFromPropertyDataArray(const UObject* SourceObject, const UAudioPropertiesSheetAsset* ParentSheet, TArrayView<AudioPropertiesSheetAssetBuilder::FPropertyRequestPtr> InPropertyData, AudioPropertiesSheetAssetBuilder::ESourceObjectParsingType ParsingType /*= AudioPropertiesSheetAssetBuilder::ESourceObjectParsingType::UObject*/)
+{
+	const FAssetToolsModule& AssetToolsModule = FModuleManager::LoadModuleChecked<FAssetToolsModule>(TEXT("AssetTools"));
+
+	FString BasePackageName = SourceObject ? SourceObject->GetPackage()->GetPathName().Append(TEXT("_GeneratedPropertySheet")) : TEXT("/Game/GeneratedPropertySheet");
+	FString OutNewPackageName;
+	FString OutNewAssetName;
+	AssetToolsModule.Get().CreateUniqueAssetName(BasePackageName, TEXT(""), OutNewPackageName, OutNewAssetName);
+
+	const TSharedRef<SDlgPickAssetPath> NewAssetDialog =
+		SNew(SDlgPickAssetPath)
+		.Title(LOCTEXT("CreateAudioPropertiesSheet", "Create Audio Properties Sheet"))
+		.DefaultAssetPath(FText::FromString(OutNewPackageName));
+
+	if (NewAssetDialog->ShowModal() == EAppReturnType::Cancel)
+	{
+		return;
+	}
+
+	const FString PackageName = NewAssetDialog->GetFullAssetPath().ToString();
+	const FName AssetName = FName(*NewAssetDialog->GetAssetName().ToString());
+
+	UPackage* const Package = CreatePackage(*PackageName);
+		
+	// Create Property Sheet Object
+	UObject* NewPropertySheetObject = nullptr;
+		
+	if (UAudioPropertiesSheetAssetFactory* const SheetFactory = Cast<UAudioPropertiesSheetAssetFactory>(NewObject<UFactory>(GetTransientPackage(), UAudioPropertiesSheetAssetFactory::StaticClass())))
+	{
+		NewPropertySheetObject = SheetFactory->FactoryCreateNew(SheetFactory->GetSupportedClass(), Package, AssetName, RF_Public | RF_Standalone, nullptr, GWarn);
+	}
+
+	if (!NewPropertySheetObject)
+	{
+		return;
+	}
+
+	// Copy selected properties
+	UAudioPropertiesSheetAsset* const NewPropertySheetAsset = CastChecked<UAudioPropertiesSheetAsset>(NewPropertySheetObject);
+
+	if (!NewPropertySheetAsset)
+	{
+		return;
+	}
+
+	NewPropertySheetAsset->PropertiesSheet.Parent = ParentSheet;
+
+	FInstancedPropertyBag& SheetPropertyBag = NewPropertySheetAsset->PropertiesSheet.Properties;
+
+	for (const AudioPropertiesSheetAssetBuilder::FPropertyRequestPtr& PropertyRequest : InPropertyData)
+	{
+		if (!PropertyRequest->bInheritProperty)
+		{
+			continue;
+		}
+
+		const FProperty* PropertyToInherit = PropertyRequest->CachedProperty;
+
+		if (!PropertyToInherit)
+		{
+			continue;
+		}
+
+		const FName& PropertyName = PropertyToInherit->GetFName();
+
+		//We create the desc locally first to remove any inherited metadata, we only want matching name and type on the sheet
+		FPropertyBagPropertyDesc PropertyDesc(PropertyName, PropertyToInherit);
+		PropertyDesc.MetaData.Empty();
+		SheetPropertyBag.AddProperties({PropertyDesc});
+		
+		const FPropertyBagPropertyDesc* CreatedPropertyDesc = SheetPropertyBag.FindPropertyDescByName(PropertyName);
+		
+		if (!CreatedPropertyDesc)
+		{
+			UE_LOG(LogInit, Warning, TEXT("Failed to create property %s from request"), *PropertyName.ToString());
+			continue;
+		}
+
+		if(!PropertyRequest->bInheritValue)
+		{
+			continue;
+		}
+
+		FStructView SheetPropertyBagView = SheetPropertyBag.GetMutableValue();
+		void* TargetPropertyAddress = CreatedPropertyDesc->CachedProperty->ContainerPtrToValuePtr<void>(SheetPropertyBagView.GetMemory());
+		
+		if (ParsingType == AudioPropertiesSheetAssetBuilder::ESourceObjectParsingType::UObject)
+		{
+			const void* SourcePropertyAddress = PropertyToInherit->ContainerPtrToValuePtr<void>(SourceObject);
+			const bool bSuccessfulCopy = PropertyAccessUtil::CopyCompletePropertyValue(PropertyToInherit, SourcePropertyAddress, CreatedPropertyDesc->CachedProperty, TargetPropertyAddress);
+		}
+		else if (ParsingType == AudioPropertiesSheetAssetBuilder::ESourceObjectParsingType::PropertySheet)
+		{
+			const UAudioPropertiesSheetAsset* SourceAsPropertySheet = Cast<UAudioPropertiesSheetAsset>(SourceObject);
+
+			if (!ensureAlwaysMsgf(SourceAsPropertySheet, TEXT("Trying to copy properties in property sheet mode from a non property sheet asset. Property copy will be skipped")))
+			{
+				continue;
+			}
+
+			auto OnLeafPropertyVisited = [&CreatedPropertyDesc, &TargetPropertyAddress, PropertyName](const FPropertyBagPropertyDesc& PropertyDesc, const FAudioPropertiesSheet& OwningSheet)
+			{
+				if (PropertyDesc.Name == PropertyName)
+				{
+					const void* SourcePropertyCointainerAddress = CreatedPropertyDesc->CachedProperty->ContainerPtrToValuePtr<void>(OwningSheet.Properties.GetValue().GetMemory());
+					const bool bSuccessfulCopy = PropertyAccessUtil::CopyCompletePropertyValue(PropertyDesc.CachedProperty, SourcePropertyCointainerAddress, CreatedPropertyDesc->CachedProperty, TargetPropertyAddress);
+				}
+			};
+
+			AudioPropertiesUtils::VisitLeafMostProperties(*SourceAsPropertySheet, OnLeafPropertyVisited);
+		}
+		else
+		{
+			checkNoEntry()
+		}
+	}
+
+	NewPropertySheetAsset->PropertiesSheet.ReconcileProperties();
+
+	FAssetRegistryModule::AssetCreated(NewPropertySheetObject);
+
+	Package->GetOutermost()->MarkPackageDirty();
+}
+
+#undef LOCTEXT_NAMESPACE
\ No newline at end of file
diff --git a/Engine/Plugins/Experimental/AudioProperties/Source/AudioPropertiesEditor/Private/AudioPropertiesSheetAssetBuilder.h b/Engine/Plugins/Experimental/AudioProperties/Source/AudioPropertiesEditor/Private/AudioPropertiesSheetAssetBuilder.h
new file mode 100644
--- /dev/null
+++ b/Engine/Plugins/Experimental/AudioProperties/Source/AudioPropertiesEditor/Private/AudioPropertiesSheetAssetBuilder.h
@@ -0,0 +1,41 @@
+// Copyright Epic Games, Inc. All Rights Reserved.
+
+#pragma once
+
+#include "Templates/SharedPointer.h"
+
+class UAudioPropertiesSheetAsset;
+class UObject;
+class FProperty;
+
+namespace AudioPropertiesSheetAssetBuilder
+{
+	struct FPropertyRequest
+	{
+		const FProperty* CachedProperty;
+		bool bInheritProperty;
+		bool bInheritValue;
+
+		FPropertyRequest(const FProperty* InCachedProperty,
+			bool bInInheritProperty = false,
+			bool bInInheritValue = false)
+			: CachedProperty(InCachedProperty)
+			, bInheritProperty(bInInheritProperty)
+			, bInheritValue(bInInheritValue)
+		{}
+	};
+
+	enum class ESourceObjectParsingType :uint8
+	{
+		UObject, //To parse UProperties found on the source UObject 
+		PropertySheet //To parse from the property sheet contained on a source UAudioPropertiesSheetAsset 
+	};
+
+	using FPropertyRequestPtr = TSharedPtr<FPropertyRequest>;
+}
+
+class FAudioPropertiesSheetAssetBuilder
+{
+public:
+	static void BuildPropertySheetFromPropertyDataArray(const UObject* SourceObject, const UAudioPropertiesSheetAsset* ParentSheet, TArrayView<AudioPropertiesSheetAssetBuilder::FPropertyRequestPtr> InPropertyData, AudioPropertiesSheetAssetBuilder::ESourceObjectParsingType ParsingType = AudioPropertiesSheetAssetBuilder::ESourceObjectParsingType::UObject);
+};
\ No newline at end of file
diff --git a/Engine/Plugins/Experimental/AudioProperties/Source/AudioPropertiesEditor/Private/AudioPropertiesSheetAssetDetails.cpp b/Engine/Plugins/Experimental/AudioProperties/Source/AudioPropertiesEditor/Private/AudioPropertiesSheetAssetDetails.cpp
new file mode 100644
--- /dev/null
+++ b/Engine/Plugins/Experimental/AudioProperties/Source/AudioPropertiesEditor/Private/AudioPropertiesSheetAssetDetails.cpp
@@ -0,0 +1,123 @@
+// Copyright Epic Games, Inc. All Rights Reserved.
+
+#include "AudioPropertiesSheetAssetDetails.h"
+
+#include "AudioPropertiesParserBase.h"
+#include "AudioPropertiesSheet.h"
+#include "DetailLayoutBuilder.h"
+#include "DetailWidgetRow.h"
+#include "IDetailGroup.h"
+#include "IPropertyUtilities.h"
+#include "PropertyCustomizationHelpers.h"
+#include "Widgets/SToolTip.h"
+#include "Widgets/Text/STextBlock.h"
+
+#define LOCTEXT_NAMESPACE "FAudioPropertiesSheetAssetDetails"
+
+namespace AudioPropertiesSheetAssetDetailsPrivate
+{
+	TSharedRef<SToolTip> CreatePathTooltip(const UObject* TargetObject)
+	{
+		FText OwnerTooltip = LOCTEXT("AudioPropertiesSheetAssetDetailsPathTooltip", "None");
+
+		if (TargetObject)
+		{
+			if (TargetObject->GetPackage())
+			{
+				OwnerTooltip = FText::FromString(TargetObject->GetPackage()->GetPathName());
+
+			}
+		}
+
+		return SNew(SToolTip)
+			[
+				SNew(STextBlock)
+				.Text(OwnerTooltip)
+			];
+	}
+}
+
+TSharedRef<IDetailCustomization> FAudioPropertiesSheetAssetDetails::MakeInstance()
+{
+	return MakeShareable(new FAudioPropertiesSheetAssetDetails);
+}
+
+void FAudioPropertiesSheetAssetDetails::CustomizeDetails(IDetailLayoutBuilder& DetailBuilder)
+{
+	BuildPropertiesParserRow(DetailBuilder);
+}
+
+void FAudioPropertiesSheetAssetDetails::BuildPropertiesParserRow(IDetailLayoutBuilder& DetailBuilder)
+{
+	TArray<TWeakObjectPtr<UObject>> Objects;
+	DetailBuilder.GetObjectsBeingCustomized(Objects);
+	
+	//We should have only one object being customized for this details view
+	if (!ensureMsgf((Objects.Num() == 1 && Objects[0].IsValid()), TEXT("Expected one valid object in the details view")))
+	{
+		return;
+	}
+
+	UAudioPropertiesSheetAsset* SheetAsset = CastChecked<UAudioPropertiesSheetAsset>(Objects[0].Get());
+
+	if (!SheetAsset)
+	{
+		return;
+	}
+
+	//To keep the custom parser rows in place we group them into a category because...
+	//...we can sort Categories but not rows 
+	IDetailCategoryBuilder& DefaultCategory = DetailBuilder.EditCategory("Default");
+	DefaultCategory.SetSortOrder(0);
+	IDetailCategoryBuilder& ParserCategory = DetailBuilder.EditCategory("Parser");
+	ParserCategory.SetSortOrder(1);
+
+	//Hide and read the local parser property so it stays on top of the category
+	TSharedRef<IPropertyHandle> LocalParserPropertyHandle = DetailBuilder.GetProperty(GET_MEMBER_NAME_CHECKED(UAudioPropertiesSheetAsset, PropertiesParser));
+	DetailBuilder.HideProperty(LocalParserPropertyHandle);
+	ParserCategory.AddProperty(LocalParserPropertyHandle);
+
+	//if we have a valid parser, we don't need to show an inherited one
+	if (SheetAsset->PropertiesParser)
+	{
+		return;
+	}
+
+	AudioPropertiesSheet::ParserParentPair ClosestParser = SheetAsset->FindClosestParserInInheritanceTree();
+	TObjectPtr<const UAudioPropertiesParserBase> InheritedParser = ClosestParser.Key;
+	const UAudioPropertiesSheetAsset* InheritedParserOwner = ClosestParser.Value;
+
+	if (!(InheritedParser && InheritedParserOwner))
+	{
+		return;
+	}
+
+	IDetailPropertyRow* InheritedParserPropertyRow = ParserCategory.AddExternalObjectProperty({ const_cast<UAudioPropertiesSheetAsset*>(InheritedParserOwner) }, GET_MEMBER_NAME_CHECKED(UAudioPropertiesSheetAsset, PropertiesParser));
+	
+	TSharedPtr<SWidget> InheritedParserDefaultNameWidget;
+	TSharedPtr<SWidget> InheritedParserDefaultValueWidget;
+	InheritedParserPropertyRow->GetDefaultWidgets(InheritedParserDefaultNameWidget, InheritedParserDefaultValueWidget);
+	
+	FString ParserOwnerName;
+	InheritedParserOwner->GetName(ParserOwnerName);
+	const FString ParentAssetPropertyString = FString::Printf(TEXT("[%s] Inherited Parser"), *ParserOwnerName);
+	FText InheritedParserNameText = FText::FromString(ParentAssetPropertyString);
+
+	TSharedRef<SWidget> InheritedParserNameWidget = SNew(STextBlock)
+		.Text(InheritedParserNameText)
+		.Font(FAppStyle::GetFontStyle(TEXT("PropertyWindow.NormalFont")))
+		.ToolTip(AudioPropertiesSheetAssetDetailsPrivate::CreatePathTooltip(InheritedParserOwner));
+
+	InheritedParserPropertyRow->CustomWidget(true)
+		.NameContent()[
+			InheritedParserNameWidget
+		]
+		.ValueContent()[
+			InheritedParserDefaultValueWidget.ToSharedRef()
+		];
+
+	//we don't want to allow editing of the inherited parser
+	InheritedParserPropertyRow->IsEnabled(false);
+}
+
+#undef LOCTEXT_NAMESPACE
\ No newline at end of file
diff --git a/Engine/Plugins/Experimental/AudioProperties/Source/AudioPropertiesEditor/Private/AudioPropertiesSheetAssetDetails.h b/Engine/Plugins/Experimental/AudioProperties/Source/AudioPropertiesEditor/Private/AudioPropertiesSheetAssetDetails.h
new file mode 100644
--- /dev/null
+++ b/Engine/Plugins/Experimental/AudioProperties/Source/AudioPropertiesEditor/Private/AudioPropertiesSheetAssetDetails.h
@@ -0,0 +1,20 @@
+// Copyright Epic Games, Inc. All Rights Reserved.
+
+#pragma once
+
+#include "IDetailCustomization.h"
+#include "Templates/SharedPointer.h"
+
+class FAudioPropertiesSheetAssetDetails : public IDetailCustomization
+{
+	public:
+	/** Makes a new instance of this detail layout class for a specific detail view requesting it */
+	static TSharedRef<IDetailCustomization> MakeInstance();
+
+private:
+	/** IDetailCustomization interface */
+	virtual void CustomizeDetails(IDetailLayoutBuilder& DetailBuilder) override;
+
+	/** Allows to select the local parser and shows the inherited one in case local is null */
+	void BuildPropertiesParserRow(IDetailLayoutBuilder& DetailBuilder);
+};
\ No newline at end of file
diff --git a/Engine/Plugins/Experimental/AudioProperties/Source/AudioPropertiesEditor/Private/AudioPropertiesSheetAssetFactory.cpp b/Engine/Plugins/Experimental/AudioProperties/Source/AudioPropertiesEditor/Private/AudioPropertiesSheetAssetFactory.cpp
new file mode 100644
--- /dev/null
+++ b/Engine/Plugins/Experimental/AudioProperties/Source/AudioPropertiesEditor/Private/AudioPropertiesSheetAssetFactory.cpp
@@ -0,0 +1,22 @@
+// Copyright Epic Games, Inc. All Rights Reserved.
+
+#include "AudioPropertiesSheetAssetFactory.h"
+
+#include "AudioPropertiesSheet.h"
+
+#include UE_INLINE_GENERATED_CPP_BY_NAME(AudioPropertiesSheetAssetFactory)
+
+UAudioPropertiesSheetAssetFactory::UAudioPropertiesSheetAssetFactory(const FObjectInitializer& ObjectInitializer)
+	: Super(ObjectInitializer)
+{
+	SupportedClass = UAudioPropertiesSheetAsset::StaticClass();
+	bCreateNew = true;
+	bEditorImport = false;
+	bEditAfterNew = true;
+}
+
+
+UObject* UAudioPropertiesSheetAssetFactory::FactoryCreateNew(UClass* Class, UObject* InParent, FName Name, EObjectFlags Flags, UObject* Context, FFeedbackContext* Warn)
+{
+	return NewObject<UAudioPropertiesSheetAsset>(InParent, Name, Flags);
+}
diff --git a/Engine/Plugins/Experimental/AudioProperties/Source/AudioPropertiesEditor/Private/AudioPropertiesSheetBuilderInstantiator.cpp b/Engine/Plugins/Experimental/AudioProperties/Source/AudioPropertiesEditor/Private/AudioPropertiesSheetBuilderInstantiator.cpp
new file mode 100644
--- /dev/null
+++ b/Engine/Plugins/Experimental/AudioProperties/Source/AudioPropertiesEditor/Private/AudioPropertiesSheetBuilderInstantiator.cpp
@@ -0,0 +1,131 @@
+// Copyright Epic Games, Inc. All Rights Reserved.
+
+#include "AudioPropertiesSheetBuilderInstantiator.h"
+
+#include "Algo/AnyOf.h"
+#include "AudioPropertiesEditorLogCategory.h"
+#include "AudioPropertiesSheet.h"
+#include "AssetDefinitionRegistry.h"
+#include "AssetDefinition_AudioPropertiesSheetAsset.h"
+#include "AssetTypeActions/AssetDefinition_SoundBase.h"
+#include "ContentBrowserMenuContexts.h"
+#include "Engine/AssetManager.h"
+#include "Engine/StreamableManager.h"
+#include "Framework/Application/SlateApplication.h"
+#include "SAudioPropertiesSheetBuilderWidget.h"
+#include "ToolMenu.h"
+#include "ToolMenus.h"
+#include "ToolMenuSection.h"
+#include "Widgets/SWindow.h"
+
+#define LOCTEXT_NAMESPACE "AudioPropertiesSheetBuilder"
+
+namespace AudioPropertiesBuilderPrivate
+{
+	auto IsPropertySheet = [](const FAssetData& AssetData) { return AssetData.IsInstanceOf<UAudioPropertiesSheetAssetBase>(); };
+}
+
+void FAudioPropertiesSheetBuilderInstantiator::ExtendContentBrowserSelectionMenu()
+{
+	if (!UToolMenus::IsToolMenuUIEnabled())
+	{
+		return;
+	}
+	
+	const UAssetDefinitionRegistry* DefRegistry = UAssetDefinitionRegistry::Get();
+	check(DefRegistry);
+
+	TArray<TObjectPtr<UAssetDefinition>> AssetDefinitions = DefRegistry->GetAllAssetDefinitions();
+	for (TObjectPtr<UAssetDefinition> AssetDef : AssetDefinitions)
+	{
+		if (UAssetDefinition_SoundBase* SoundAssetDef = Cast<UAssetDefinition_SoundBase>(AssetDef.Get()))
+		{
+			FToolMenuSection* Section = SoundAssetDef->FindSoundContextMenuSection("Sound");
+			check(Section);
+
+			Section->AddDynamicEntry("AudioPropertiesSheet_CreatePropertiesSheetBuilder", FNewToolMenuSectionDelegate::CreateLambda([this](FToolMenuSection& InSection)
+			{
+				if (const UContentBrowserAssetContextMenuContext* Context = InSection.FindContext<UContentBrowserAssetContextMenuContext>())
+				{
+					const TAttribute<FText> Label = LOCTEXT("AudioPropertiesSheetBuilderMenu_Label", "Open Property Sheet Builder");
+					const TAttribute<FText> ToolTip = LOCTEXT("AudioPropertiesSheetBuilderMenu_Tooltip", "Opens a utility for creating property sheets with properties from the selected objects.");
+					const FSlateIcon Icon = FSlateIcon();
+					const FToolMenuExecuteAction UIAction = FToolMenuExecuteAction::CreateSP(this, &FAudioPropertiesSheetBuilderInstantiator::ExecuteCreateBuilderWidget);
+		
+					InSection.AddMenuEntry("AudioPropertiesSheetBuilderMenu_Label", Label, ToolTip, Icon, UIAction);
+					
+				}
+			}));
+		}
+		else if (UAssetDefinition_AudioPropertiesSheetAsset* AudioPropertySheetAssetDef = Cast<UAssetDefinition_AudioPropertiesSheetAsset>(AssetDef.Get()))
+		{
+			FToolMenuSection* Section = AudioPropertySheetAssetDef->FindContextMenuSection("Builder");
+			check(Section);
+	
+			Section->AddDynamicEntry("AudioPropertiesSheet_CreateChildPropertySheet", FNewToolMenuSectionDelegate::CreateLambda([this](FToolMenuSection& InSection)
+			{
+				if (const UContentBrowserAssetContextMenuContext* Context = InSection.FindContext<UContentBrowserAssetContextMenuContext>())
+				{
+					if (Algo::AnyOf(Context->SelectedAssets, AudioPropertiesBuilderPrivate::IsPropertySheet))
+					{
+						const TAttribute<FText> Label = LOCTEXT("AudioPropertiesSheetChildBuilderMenu_Label", "Open Property Sheet Child Builder");
+						const TAttribute<FText> ToolTip = LOCTEXT("AudioPropertiesSheetChildBuilderMenu_Tooltip", "Opens a utility for creating a child of this property sheet");
+						const FSlateIcon Icon = FSlateIcon();
+						const FToolMenuExecuteAction UIAction = FToolMenuExecuteAction::CreateSP(this, &FAudioPropertiesSheetBuilderInstantiator::ExecuteCreateBuilderWidget);
+	
+						InSection.AddMenuEntry("AudioPropertiesSheetBuilderMenu_Label", Label, ToolTip, Icon, UIAction);
+					}
+				}
+			}));
+		}
+	}	
+}
+
+void FAudioPropertiesSheetBuilderInstantiator::ExecuteCreateBuilderWidget(const FToolMenuContext& MenuContext)
+{
+	if (const UContentBrowserAssetContextMenuContext* Context = UContentBrowserAssetContextMenuContext::FindContextWithAssets(MenuContext))
+	{
+		TArray<TSoftObjectPtr<UObject>> SelectedObjects = Context->GetSelectedAssetSoftObjects<UObject>();
+		CreateBuilderWidget(SelectedObjects);
+	}
+}
+
+void FAudioPropertiesSheetBuilderInstantiator::CreateBuilderWidget(TArrayView<TSoftObjectPtr<UObject>> SourceObjects)
+{
+	for (TSoftObjectPtr<UObject> SourceObject : SourceObjects)
+	{
+		UE_LOG(LogAudioPropertiesEditor, Log, TEXT("Loading Object %s to open property sheet builder"), *SourceObject->GetPathName())
+		FStreamableDelegate LibDelegate = FStreamableDelegate::CreateSP(this, &FAudioPropertiesSheetBuilderInstantiator::SourceObjectAsyncLoadComplete, SourceObject);
+		UAssetManager::GetStreamableManager().RequestAsyncLoad(SourceObject.ToSoftObjectPath(), MoveTemp(LibDelegate));
+	}
+}
+
+void FAudioPropertiesSheetBuilderInstantiator::SourceObjectAsyncLoadComplete(TSoftObjectPtr<UObject> LoadedObject)
+{
+	if (!LoadedObject.IsValid())
+	{
+		return;
+	}
+
+	UE_LOG(LogAudioPropertiesEditor, Log, TEXT("Opening property sheet builder with source %s"), *LoadedObject->GetPathName())
+	
+	TObjectPtr<UObject> LoadedSource = LoadedObject.Get();
+	const UClass* SourceClass = LoadedSource.GetClass();
+	
+	TSharedRef<SWindow> PropertiesWindow = SNew(SWindow)
+	.Title(FText::FromString("Audio Properties Sheet Builder"))
+	.ClientSize(FVector2D(800, 600))
+	.SupportsMinimize(false)
+	.SupportsMaximize(false);
+
+	PropertiesWindow->SetContent(
+		SNew(SAudioPropertiesSheetBuilderWidget)
+			.SourceObject(LoadedSource.Get())
+			.bSourceIsParent(SourceClass && SourceClass->IsChildOf(UAudioPropertiesSheetAsset::StaticClass()))
+	);
+
+	FSlateApplication::Get().AddWindow(PropertiesWindow);
+}
+
+
+#undef LOCTEXT_NAMESPACE
diff --git a/Engine/Plugins/Experimental/AudioProperties/Source/AudioPropertiesEditor/Private/AudioPropertiesSheetBuilderInstantiator.h b/Engine/Plugins/Experimental/AudioProperties/Source/AudioPropertiesEditor/Private/AudioPropertiesSheetBuilderInstantiator.h
new file mode 100644
--- /dev/null
+++ b/Engine/Plugins/Experimental/AudioProperties/Source/AudioPropertiesEditor/Private/AudioPropertiesSheetBuilderInstantiator.h
@@ -0,0 +1,20 @@
+// Copyright Epic Games, Inc. All Rights Reserved.
+
+#pragma once
+
+#include "Templates/SharedPointer.h"
+#include "UObject/SoftObjectPtr.h"
+
+class UObject;
+struct FToolMenuContext;
+
+class FAudioPropertiesSheetBuilderInstantiator : public TSharedFromThis<FAudioPropertiesSheetBuilderInstantiator>
+{
+public:
+	void ExtendContentBrowserSelectionMenu();
+
+private:
+	void ExecuteCreateBuilderWidget(const FToolMenuContext& MenuContext);
+	void CreateBuilderWidget(TArrayView<TSoftObjectPtr<UObject>> SourceObjects);
+	void SourceObjectAsyncLoadComplete(TSoftObjectPtr<UObject> LoadedObject);
+};
diff --git a/Engine/Plugins/Experimental/AudioProperties/Source/AudioPropertiesEditor/Private/AudioPropertiesSheetDetails.cpp b/Engine/Plugins/Experimental/AudioProperties/Source/AudioPropertiesEditor/Private/AudioPropertiesSheetDetails.cpp
new file mode 100644
--- /dev/null
+++ b/Engine/Plugins/Experimental/AudioProperties/Source/AudioPropertiesEditor/Private/AudioPropertiesSheetDetails.cpp
@@ -0,0 +1,487 @@
+// Copyright Epic Games, Inc. All Rights Reserved.
+
+#include "AudioPropertiesSheetDetails.h"
+
+#include "AudioPropertiesEditorLogCategory.h"
+#include "AudioPropertiesSheet.h"
+#include "AudioPropertiesUtils.h"
+#include "DetailLayoutBuilder.h"
+#include "DetailWidgetRow.h"
+#include "Editor.h"
+#include "Framework/Application/SlateApplication.h"
+#include "IDetailChildrenBuilder.h"
+#include "InstancedAudioPropertyBagDetails.h"
+#include "IPropertyChangeListener.h"
+#include "IPropertyUtilities.h"
+#include "PropertyBagDetails.h"
+#include "PropertyCustomizationHelpers.h"
+#include "PropertyEditorModule.h"
+#include "ThumbnailRendering/ThumbnailManager.h"
+#include "Widgets/Input/SComboButton.h"
+#include "Widgets/SNullWidget.h"
+#include "Widgets/SToolTip.h"
+#include "Widgets/Text/STextBlock.h"
+
+#define LOCTEXT_NAMESPACE "FAudioPropertiesSheetDetails"
+
+namespace AudioPropertiesSheetDetailsPrivate
+{
+	const UObject* GetUObjectFromParentPtrPropertyHandle(const TSharedPtr<IPropertyHandle> InPropertyHandle)
+	{
+		check(InPropertyHandle)
+
+		if (FProperty* PropertyPtr = InPropertyHandle->GetProperty())
+		{
+			if (FObjectProperty* ObjectProperty = CastField<FObjectProperty>(PropertyPtr))
+			{
+				void* ObjPropertyPtr = nullptr;
+				InPropertyHandle->GetValueData(ObjPropertyPtr);
+				if (ObjPropertyPtr)
+				{
+					TObjectPtr<const UObject>* ObjPtr = (TObjectPtr<const UObject>*)ObjPropertyPtr;
+					return ObjPtr->Get();
+				}
+			}
+		}
+
+		return nullptr;
+
+	}
+
+	static const FName ParentPropertyName = GET_MEMBER_NAME_CHECKED(FAudioPropertiesSheet, Parent);
+	static const FName PropertyBagPropertyName = GET_MEMBER_NAME_CHECKED(FAudioPropertiesSheet, Properties);
+}
+
+TSharedRef<IPropertyTypeCustomization> FAudioPropertiesSheetDetails::MakeInstance()
+{
+	return MakeShareable(new FAudioPropertiesSheetDetails);
+}
+
+void FAudioPropertiesSheetDetails::CustomizeHeader(TSharedRef<IPropertyHandle> InPropertyHandle, FDetailWidgetRow& InHeaderRow, IPropertyTypeCustomizationUtils& InCustomizationUtils)
+{
+}
+
+void FAudioPropertiesSheetDetails::CustomizeChildren(TSharedRef<IPropertyHandle> InPropertyHandle, IDetailChildrenBuilder& InChildrenBuilder, IPropertyTypeCustomizationUtils& InCustomizationUtils)
+{
+	FAudioPropertiesSheet* PropertiesSheetStruct = nullptr; 
+	TArray<UObject*> OuterObjects;
+	InPropertyHandle->GetOuterObjects(OuterObjects);
+	check(OuterObjects.Num() > 0 && OuterObjects[0]) //multi edit not supported yet
+
+	UAudioPropertiesSheetAsset* Owner = Cast<UAudioPropertiesSheetAsset>(OuterObjects[0]);
+	if (Owner)
+	{
+		PropertiesSheetStruct = &Owner->PropertiesSheet;
+
+		const TSharedRef<FAudioPropertiesSheetBagNodeBuilder> BagNodeBuilder = MakeShareable(new FAudioPropertiesSheetBagNodeBuilder(InPropertyHandle, InCustomizationUtils.GetPropertyUtilities(), PropertiesSheetStruct));
+		InChildrenBuilder.AddCustomBuilder(BagNodeBuilder);
+
+		const bool bIsLeaf = true;
+		const TSharedRef<FAudioPropertiesSheetAssetTreeNodeBuilder> AssetTreeNodeBuilder = MakeShareable(new FAudioPropertiesSheetAssetTreeNodeBuilder(InPropertyHandle, InPropertyHandle->GetChildHandle(AudioPropertiesSheetDetailsPrivate::ParentPropertyName), PropertiesSheetStruct, bIsLeaf));
+		InChildrenBuilder.AddCustomBuilder(AssetTreeNodeBuilder);
+	}
+}
+
+FAudioPropertiesSheetBagNodeBuilder::FAudioPropertiesSheetBagNodeBuilder(TSharedPtr<IPropertyHandle> InPropertySheetStructHandle, TSharedPtr<IPropertyUtilities> InPropertyUtils, FAudioPropertiesSheet* InLeafMostSheet /*= nullptr*/)
+	: CachedSheetStructHandle(InPropertySheetStructHandle)
+	, CachedPropertyUtils(InPropertyUtils)
+	, LeafMostSheet(InLeafMostSheet)
+{
+	if (CachedSheetStructHandle)
+	{
+		ParentPropertyHandle = CachedSheetStructHandle->GetChildHandle(AudioPropertiesSheetDetailsPrivate::ParentPropertyName);
+	}
+}
+
+void FAudioPropertiesSheetBagNodeBuilder::GenerateHeaderRowContent(FDetailWidgetRow& NodeRow)
+{
+	NodeRow.NameContent()[BuildHeaderRowName()]
+			.ValueContent()[BuildHeaderRowValue()]
+			.ShouldAutoExpand(true);
+}
+
+FName FAudioPropertiesSheetBagNodeBuilder::GetName() const
+{
+	static const FName AudioPropertiesSheetBagNodeName("AudioPropertiesSheetBagNode");
+	return AudioPropertiesSheetBagNodeName;
+}
+
+void FAudioPropertiesSheetBagNodeBuilder::OnParentAssetPropertyOverrideChange(const FProperty& TargetProperty, bool bIsOverridden)
+{
+	if (CachedPropertyUtils) 
+	{
+		CachedPropertyUtils->ForceRefresh();
+	}
+}
+
+void FAudioPropertiesSheetBagNodeBuilder::GenerateChildContent(IDetailChildrenBuilder& ChildrenBuilder)
+{
+	if (TSharedPtr<IPropertyHandle> LocalPropertyBagHandle = CachedSheetStructHandle->GetChildHandle(AudioPropertiesSheetDetailsPrivate::PropertyBagPropertyName))
+	{
+		const bool bIsLeafBag = true;
+
+		const TSharedRef<FInstancedAudioPropertyBagDetails> LocalBagDetails = MakeShareable(new FInstancedAudioPropertyBagDetails(LocalPropertyBagHandle, CachedPropertyUtils, nullptr, LeafMostSheet, bIsLeafBag));
+		ChildrenBuilder.AddCustomBuilder(LocalBagDetails);
+	}
+
+	AddParentTreePropertyBags(ChildrenBuilder);
+}
+
+void FAudioPropertiesSheetBagNodeBuilder::AddParentTreePropertyBags(IDetailChildrenBuilder& ChildrenBuilder)
+{
+	const UObject* ParentAsset = AudioPropertiesSheetDetailsPrivate::GetUObjectFromParentPtrPropertyHandle(ParentPropertyHandle);
+
+	while (ParentAsset)
+	{
+		TArray<UObject*> ParentObjects;
+		ParentObjects.Add(const_cast<UObject*>(ParentAsset));
+
+		UAudioPropertiesSheetAsset* ParentAsSheet = Cast<UAudioPropertiesSheetAsset>(ParentObjects[0]);
+				
+		if (ParentAsSheet)
+		{
+			auto ParentObjectsRow = ChildrenBuilder.AddExternalObjectProperty(ParentObjects, TEXT("PropertiesSheet"));
+			auto ParentRowPropertyHandle = ParentObjectsRow->GetPropertyHandle();
+			ParentObjectsRow->Visibility(EVisibility::Hidden);
+
+			if (TSharedPtr<IPropertyHandle> PropertyBagHandle = ParentRowPropertyHandle->GetChildHandle(AudioPropertiesSheetDetailsPrivate::PropertyBagPropertyName))
+			{
+				const bool bIsLeafBag = false;
+
+				TSharedRef<FInstancedAudioPropertyBagDetails> InheritedBagDetails = MakeShareable(new FInstancedAudioPropertyBagDetails(PropertyBagHandle, CachedPropertyUtils, ParentAsSheet, LeafMostSheet, bIsLeafBag));
+				ChildrenBuilder.AddCustomBuilder(InheritedBagDetails);
+			
+				ParentAsSheet->PropertiesSheet.OnAudioPropertyOverrideChange.AddSP( this, &FAudioPropertiesSheetBagNodeBuilder::OnParentAssetPropertyOverrideChange );
+
+			}
+
+			ParentAsset = ParentAsSheet->PropertiesSheet.Parent;
+
+		}
+		else
+		{
+			ParentAsset = nullptr;
+		}		
+	}
+}
+
+TSharedRef<SWidget> FAudioPropertiesSheetBagNodeBuilder::BuildHeaderRowName() const
+{
+	FText HeaderValueText =  LOCTEXT("AudioPropertiesSheetBagNodeHeaderName", "Properties");
+
+	return SNew(STextBlock)
+		.Text(HeaderValueText)
+		.Font(FAppStyle::GetFontStyle(TEXT("PropertyWindow.NormalFont")));
+}
+
+TSharedRef<SWidget> FAudioPropertiesSheetBagNodeBuilder::BuildHeaderRowValue()
+{
+	if (TSharedPtr<IPropertyHandle> LocalPropertyBagHandle = CachedSheetStructHandle->GetChildHandle(AudioPropertiesSheetDetailsPrivate::PropertyBagPropertyName))
+	{
+		return FPropertyBagDetails::MakeAddPropertyWidget(LocalPropertyBagHandle, CachedPropertyUtils).ToSharedRef();
+	}
+
+	return SNew(STextBlock);
+}
+
+FAudioPropertiesSheetAssetTreeNodeBuilder::FAudioPropertiesSheetAssetTreeNodeBuilder(TSharedPtr<IPropertyHandle> InPropertySheetStructHandle, TSharedPtr<IPropertyHandle> InChildPtrToThisContainer, FAudioPropertiesSheet* InLeafMostSheet /*= nullptr*/, bool InIsLeaf /*= false*/)
+	: CachedSheetStructHandle(InPropertySheetStructHandle)
+	, ChildPtrToThisContainer(InChildPtrToThisContainer)
+	, LeafMostSheet(InLeafMostSheet)
+	, bIsLeaf(InIsLeaf)
+{
+	if (CachedSheetStructHandle)
+	{
+		CachedSheetStructHandle->GetOuterObjects(PropertyOwners);
+		ParentPropertyHandle = CachedSheetStructHandle->GetChildHandle(AudioPropertiesSheetDetailsPrivate::ParentPropertyName);
+	}
+}
+
+void FAudioPropertiesSheetAssetTreeNodeBuilder::GenerateHeaderRowContent(FDetailWidgetRow& NodeRow)
+{	
+	const UObject* ParentAsset = AudioPropertiesSheetDetailsPrivate::GetUObjectFromParentPtrPropertyHandle(ParentPropertyHandle);
+	
+	NodeRow.NameContent()[BuildHeaderRowName()]
+		.ValueContent()[BuildHeaderRowValue(ParentAsset, bIsLeaf)]
+		.ShouldAutoExpand(true);		
+}
+
+void FAudioPropertiesSheetAssetTreeNodeBuilder::GenerateChildContent(IDetailChildrenBuilder& ChildrenBuilder)
+{
+	AddParentNode(ChildrenBuilder);
+}
+
+void FAudioPropertiesSheetAssetTreeNodeBuilder::AddParentNode(IDetailChildrenBuilder& ChildrenBuilder)
+{
+	if (const UObject* ParentAsset = AudioPropertiesSheetDetailsPrivate::GetUObjectFromParentPtrPropertyHandle(ParentPropertyHandle))
+	{
+		TArray<UObject*> ParentObjects;
+		ParentObjects.Add(const_cast<UObject*>(ParentAsset));
+
+		auto ParentObjectsRow = ChildrenBuilder.AddExternalObjectProperty(ParentObjects, TEXT("PropertiesSheet"));
+		auto ParentRowPropertyHandle = ParentObjectsRow->GetPropertyHandle();
+		ParentObjectsRow->Visibility(EVisibility::Hidden);
+
+		constexpr bool bIsParentLeaf = false;
+		const TSharedRef<FAudioPropertiesSheetAssetTreeNodeBuilder> ParentDetails = MakeShareable(new FAudioPropertiesSheetAssetTreeNodeBuilder(ParentRowPropertyHandle, ParentPropertyHandle, LeafMostSheet, bIsParentLeaf));
+		ChildrenBuilder.AddCustomBuilder(ParentDetails);
+	}
+}
+
+FName FAudioPropertiesSheetAssetTreeNodeBuilder::GetName() const
+{
+	static const FName AudioPropertiesSheetAssetTreeNodeName("AudioPropertiesSheetAssetTreeNode");
+	return AudioPropertiesSheetAssetTreeNodeName;
+}
+
+TSharedRef<SWidget> FAudioPropertiesSheetAssetTreeNodeBuilder::BuildParentPickerWidget()
+{
+	FOnGetAllowedClasses AssetPickerAllowedClasses = FOnGetAllowedClasses::CreateLambda([](TArray<const UClass*>& OutClasses) {OutClasses.Add(UAudioPropertiesSheetAsset::StaticClass()); });
+	FOnAssetSelected AssetPickerAssetSelected = FOnAssetSelected::CreateSP(this, &FAudioPropertiesSheetAssetTreeNodeBuilder::OnParentAssetSelected);
+	FOnShouldFilterAsset AssetPickerFilter = FOnShouldFilterAsset::CreateSP(this, &FAudioPropertiesSheetAssetTreeNodeBuilder::FilterParentAssetForPicking);
+	FSimpleDelegate AssetPickerClose = FSimpleDelegate::CreateSP(this, &FAudioPropertiesSheetAssetTreeNodeBuilder::OnAssetPickerClosed);
+
+	return PropertyCustomizationHelpers::MakeAssetPickerWithMenu(
+		FAssetData(),
+		true,
+		true,
+		{ UAudioPropertiesSheetAsset::StaticClass() },
+		TArray<UFactory*>(),
+		AssetPickerFilter,
+		AssetPickerAssetSelected,
+		AssetPickerClose,
+		ParentPropertyHandle);
+}
+
+TSharedRef<SWidget> FAudioPropertiesSheetAssetTreeNodeBuilder::BuildHeaderRowName() const
+{
+	FText HeaderValueText;
+
+
+	auto GenerateOwningSheetText = [this]()
+	{
+		if (bIsLeaf || !PropertyOwners[0])
+		{
+			return FText();
+		}
+		
+		FString OwningAssetName;
+		PropertyOwners[0]->GetName(OwningAssetName);
+		const FString ParentAssetPropertyString = FString::Printf(TEXT("[%s] "), *OwningAssetName);
+
+		return FText::FromString(ParentAssetPropertyString);
+	};
+
+	return SNew(SHorizontalBox)
+				+ SHorizontalBox::Slot()
+				.AutoWidth()
+				[
+					SNew(STextBlock)
+					.Text_Lambda(GenerateOwningSheetText)
+					.Font(IDetailLayoutBuilder::GetDetailFontItalic())
+				]
+				+ SHorizontalBox::Slot()
+				.AutoWidth()
+				[
+					SNew(STextBlock)
+					.Font(IDetailLayoutBuilder::GetDetailFont())
+					.Text(LOCTEXT("AudioPropertiesSheetParentRowName", "Parent"))
+				];
+}
+
+TSharedRef<SWidget> FAudioPropertiesSheetAssetTreeNodeBuilder::BuildHeaderRowValue(const UObject* PropertyOwner, const bool InBuildAssetPicker)
+{
+	check(CachedSheetStructHandle)
+
+	TSharedPtr<SWidget> AssetRowValue = nullptr;
+	FText OwnerNameText = PropertyOwner ?  FText::FromString(PropertyOwner->GetName()) : LOCTEXT("AudioPropertiesSheetNodeDetailsHeaderValue", "None");
+
+
+	TSharedRef<SWidget> ValueTextWidget = SNew(STextBlock).Text(OwnerNameText)
+		.Font(FAppStyle::GetFontStyle(TEXT("PropertyWindow.BoldFont")))
+		.ToolTip(CreatePathTooltip(PropertyOwner));
+
+
+	if (InBuildAssetPicker)
+	{
+		TSharedRef<SWidget> ComboBoxContent = SNew(SHorizontalBox)
+		+ SHorizontalBox::Slot()
+		.FillWidth(1.0f)
+		.VAlign(VAlign_Center)
+		.Padding(5, 0)
+		[
+			ValueTextWidget
+		];
+
+		AssetRowValue = SNew(SHorizontalBox)
+		+ SHorizontalBox::Slot()
+		.VAlign(VAlign_Center)
+		.AutoWidth()
+		[
+			SNew(SComboButton)
+			.OnGetMenuContent(this, &FAudioPropertiesSheetAssetTreeNodeBuilder::BuildParentPickerWidget)
+			.ButtonContent()
+			[
+				ComboBoxContent
+			]
+		]
+		+SHorizontalBox::Slot()
+		.AutoWidth()
+		.Padding(2.0f, 0.0f)
+		.VAlign(VAlign_Center)
+		[
+			PropertyCustomizationHelpers::MakeUseSelectedButton(
+			FSimpleDelegate::CreateSP(this, &FAudioPropertiesSheetAssetTreeNodeBuilder::OnUseSelected))
+		]
+		+SHorizontalBox::Slot()
+			.AutoWidth()
+			.Padding(2.0f, 0.0f)
+			.VAlign(VAlign_Center)
+			[
+				PropertyCustomizationHelpers::MakeBrowseButton(
+				FSimpleDelegate::CreateSP(this, &FAudioPropertiesSheetAssetTreeNodeBuilder::OnBrowseTo, const_cast<UObject*>(PropertyOwner)))
+			]
+		+ SHorizontalBox::Slot()
+			.AutoWidth()
+			.Padding(2.0f, 0.0f)
+			.VAlign(VAlign_Center)
+			[
+				PropertyCustomizationHelpers::MakeClearButton(
+				FSimpleDelegate::CreateSP(this, &FAudioPropertiesSheetAssetTreeNodeBuilder::OnParentCleared))
+			];
+	}
+	else
+	{
+		 TSharedRef<SHorizontalBox> AssetRowBox = SNew(SHorizontalBox)
+			+ SHorizontalBox::Slot()
+			.VAlign(VAlign_Center)
+			.AutoWidth()
+			[
+				ValueTextWidget
+			];
+
+			if (PropertyOwner)
+			{
+				AssetRowBox->AddSlot().VAlign(VAlign_Center)
+				.AutoWidth()
+				[
+					PropertyCustomizationHelpers::MakeBrowseButton(
+					FSimpleDelegate::CreateSP(this, &FAudioPropertiesSheetAssetTreeNodeBuilder::OnBrowseTo, const_cast<UObject*>(PropertyOwner)))
+				];				
+			}
+
+			AssetRowValue = AssetRowBox;
+	}
+
+
+	return AssetRowValue.ToSharedRef();
+}
+
+TSharedRef<SToolTip> FAudioPropertiesSheetAssetTreeNodeBuilder::CreatePathTooltip(const UObject* TargetObject)
+{
+	FText OwnerTooltip = LOCTEXT("AudioPropertiesSheetNodeDetailsHeaderTooltip", "None");
+
+	if (TargetObject)
+	{
+		if (TargetObject->GetPackage())
+		{
+			OwnerTooltip = FText::FromString(TargetObject->GetPackage()->GetPathName());
+
+		}
+	}
+
+	return SNew(SToolTip)
+		[
+			SNew(STextBlock)
+			.Text(OwnerTooltip)
+		];
+}
+
+void FAudioPropertiesSheetAssetTreeNodeBuilder::OnUseSelected()
+{
+	TArray<FAssetData> SelectedAssets;
+	GEditor->GetContentBrowserSelections(SelectedAssets);
+
+	if (SelectedAssets.Num() > 0)
+	{
+		if (!FilterParentAssetForPicking(SelectedAssets[0]))
+		{
+			OnParentAssetSelected(SelectedAssets[0]);
+		}
+		else
+		{
+			UE_LOG(LogAudioPropertiesEditor, Log, TEXT("Trying to select invalid Audio Properties Sheet parent from content browser"))
+		}
+		
+	}
+}
+
+void FAudioPropertiesSheetAssetTreeNodeBuilder::OnBrowseTo(UObject* AssetToBrowseTo)
+{
+	TArray<UObject*> AssetsToBrowseTo;
+	AssetsToBrowseTo.Add(AssetToBrowseTo);
+	GEditor->SyncBrowserToObjects(AssetsToBrowseTo);
+}
+
+bool FAudioPropertiesSheetAssetTreeNodeBuilder::FilterParentAssetForPicking(const FAssetData& InAssetData)
+{
+	if (InAssetData.GetClass() == UAudioPropertiesSheetAsset::StaticClass())
+	{
+		TArray<UPackage*> OuterPackages;
+		ChildPtrToThisContainer->GetOuterPackages(OuterPackages);
+		
+		for (const UPackage* Package : OuterPackages)
+		{
+			if (Package->GetOutermost() == InAssetData.GetPackage())
+			{
+				//Found a matching package with the same asset data, filter it out
+				return true;
+			}
+		}
+
+		if (const UAudioPropertiesSheetAsset* PropertiesSheetAssetPtr = Cast<UAudioPropertiesSheetAsset>(InAssetData.GetAsset()))
+		{
+			TArray<UObject*> ContainerObjects;
+			ChildPtrToThisContainer->GetOuterObjects(ContainerObjects);
+			{
+				for (const UObject* Object : ContainerObjects)
+				{
+					if (const UAudioPropertiesSheetAsset* PropertyOwner = Cast<UAudioPropertiesSheetAsset>(Object))
+					{
+						if (AudioPropertiesUtils::IsSheetInInheritanceTree(*PropertiesSheetAssetPtr, PropertyOwner))
+						{
+							return true;
+						}
+					}
+				}
+			}
+		}
+
+		return false;
+	}
+
+	return true;
+}
+
+void FAudioPropertiesSheetAssetTreeNodeBuilder::OnParentAssetSelected(const FAssetData& AssetData)
+{
+	ChildPtrToThisContainer->SetValueFromFormattedString(AssetData.IsValid() ? AssetData.GetAsset()->GetPathName() : TEXT("None"));
+	ChildPtrToThisContainer->NotifyPostChange(EPropertyChangeType::ValueSet);
+	ChildPtrToThisContainer->RequestRebuildChildren();
+	OnRebuildChildren.ExecuteIfBound();
+}
+
+void FAudioPropertiesSheetAssetTreeNodeBuilder::OnParentCleared()
+{
+	ChildPtrToThisContainer->SetValueFromFormattedString(TEXT("None"));
+	ChildPtrToThisContainer->NotifyPostChange(EPropertyChangeType::ValueSet);
+	ChildPtrToThisContainer->RequestRebuildChildren();
+	OnRebuildChildren.ExecuteIfBound();
+}
+
+void FAudioPropertiesSheetAssetTreeNodeBuilder::OnAssetPickerClosed()
+{
+	FSlateApplication::Get().DismissAllMenus();
+}
+
+#undef LOCTEXT_NAMESPACE
\ No newline at end of file
diff --git a/Engine/Plugins/Experimental/AudioProperties/Source/AudioPropertiesEditor/Private/AudioPropertiesSheetDetails.h b/Engine/Plugins/Experimental/AudioProperties/Source/AudioPropertiesEditor/Private/AudioPropertiesSheetDetails.h
new file mode 100644
--- /dev/null
+++ b/Engine/Plugins/Experimental/AudioProperties/Source/AudioPropertiesEditor/Private/AudioPropertiesSheetDetails.h
@@ -0,0 +1,105 @@
+// Copyright Epic Games, Inc. All Rights Reserved.
+
+#pragma once
+
+#include "IDetailCustomNodeBuilder.h"
+#include "IPropertyTypeCustomization.h"
+
+class SToolTip;
+class SWidget;
+struct FAudioPropertiesSheet;
+
+//Details Customization of an AudioPropertiesSheetAsset
+class FAudioPropertiesSheetDetails : public IPropertyTypeCustomization
+{
+public:
+	/** Makes a new instance of this detail layout class for a specific detail view requesting it */
+	static TSharedRef<IPropertyTypeCustomization> MakeInstance();
+
+	/** ~Begin IPropertyTypeCustomization interface */
+	virtual void CustomizeHeader(TSharedRef<IPropertyHandle> InPropertyHandle, FDetailWidgetRow& InHeaderRow, IPropertyTypeCustomizationUtils& InCustomizationUtils) override;
+	virtual void CustomizeChildren(TSharedRef<IPropertyHandle> InPropertyHandle, IDetailChildrenBuilder& InChildrenBuilder, IPropertyTypeCustomizationUtils& InCustomizationUtils) override;
+	/** ~End IPropertyTypeCustomization interface */
+};
+
+//Creates a property node showing the local property bag + any properties inherited from the parent tree
+class FAudioPropertiesSheetBagNodeBuilder : public IDetailCustomNodeBuilder, public TSharedFromThis<FAudioPropertiesSheetBagNodeBuilder>
+{
+public:
+	FAudioPropertiesSheetBagNodeBuilder(
+		TSharedPtr<IPropertyHandle> InPropertySheetStructHandle,
+		TSharedPtr<IPropertyUtilities> InPropertyUtils, 
+		FAudioPropertiesSheet* InLeafMostSheet = nullptr
+	);
+
+	void GenerateChildContent( IDetailChildrenBuilder& ChildrenBuilder )override;
+	void GenerateHeaderRowContent(FDetailWidgetRow& NodeRow) override;
+
+	FName GetName() const override;
+
+	void OnParentAssetPropertyOverrideChange(const FProperty& TargetProperty, bool bIsOverridden);
+
+
+private:
+	void AddParentTreePropertyBags(IDetailChildrenBuilder& ChildrenBuilder);
+
+	TSharedRef<SWidget> BuildHeaderRowName() const;
+	TSharedRef<SWidget> BuildHeaderRowValue();
+
+	TSharedPtr<IPropertyHandle> CachedSheetStructHandle;
+	TSharedPtr<IPropertyHandle> ParentPropertyHandle;
+
+	TSharedPtr<IPropertyUtilities> CachedPropertyUtils;
+	FAudioPropertiesSheet* LeafMostSheet;
+
+	TSharedPtr<class IPropertyChangeListener> PropertyChangeListener;
+
+};
+
+//Creates a property node showing the property sheet inheritance tree
+class FAudioPropertiesSheetAssetTreeNodeBuilder : public IDetailCustomNodeBuilder, public TSharedFromThis<FAudioPropertiesSheetAssetTreeNodeBuilder>
+{
+public:
+	FAudioPropertiesSheetAssetTreeNodeBuilder(
+		TSharedPtr<IPropertyHandle> InPropertySheetStructHandle, 
+		TSharedPtr<IPropertyHandle> InChildPtrToThisContainer,
+		FAudioPropertiesSheet* InLeafMostSheet = nullptr,
+		bool InIsLeaf = false);
+
+	void GenerateHeaderRowContent(FDetailWidgetRow& NodeRow) override;
+	void GenerateChildContent(IDetailChildrenBuilder& ChildrenBuilder) override;
+
+	virtual void SetOnRebuildChildren(FSimpleDelegate InOnRegenerateChildren) override { OnRebuildChildren = InOnRegenerateChildren; }
+
+	FName GetName() const override;
+
+	FSimpleDelegate OnRebuildChildren;
+
+private:
+	TSharedRef<SWidget> BuildParentPickerWidget();
+	TSharedRef<SWidget> BuildHeaderRowName() const;
+	TSharedRef<SWidget> BuildHeaderRowValue(const UObject* PropertyOwner, const bool InBuildAssetPicker);
+	TSharedRef<SToolTip> CreatePathTooltip(const UObject* TargetObject);
+
+	void AddParentNode(IDetailChildrenBuilder& ChildrenBuilder);
+
+	TSharedPtr<IPropertyHandle> CachedSheetStructHandle;
+	TSharedPtr<IPropertyHandle> ParentPropertyHandle;
+	TSharedPtr<IPropertyHandle> ChildPtrToThisContainer;
+
+	const IPropertyTypeCustomizationUtils* CustomizationUtils;
+
+	//asset browsing/picking delegates
+	bool FilterParentAssetForPicking(const FAssetData& InAssetData);
+	void OnParentAssetSelected(const FAssetData& AssetData);
+	void OnParentCleared();
+	void OnAssetPickerClosed();
+	void OnUseSelected();
+	void OnBrowseTo(UObject* AssetToBrowseTo);
+
+	TArray<UObject*> PropertyOwners;
+
+	FAudioPropertiesSheet* LeafMostSheet;
+
+	bool bIsLeaf;
+};
\ No newline at end of file
diff --git a/Engine/Plugins/Experimental/AudioProperties/Source/AudioPropertiesEditor/Private/CustomNodeBuilders/AudioPropertiesDetailArrayBuilder.cpp b/Engine/Plugins/Experimental/AudioProperties/Source/AudioPropertiesEditor/Private/CustomNodeBuilders/AudioPropertiesDetailArrayBuilder.cpp
new file mode 100644
--- /dev/null
+++ b/Engine/Plugins/Experimental/AudioProperties/Source/AudioPropertiesEditor/Private/CustomNodeBuilders/AudioPropertiesDetailArrayBuilder.cpp
@@ -0,0 +1,144 @@
+﻿// Copyright Epic Games, Inc. All Rights Reserved.
+
+#include "AudioPropertiesDetailArrayBuilder.h"
+
+#include "AudioPropertiesDetailsInjectorUtils.h"
+#include "IDetailChildrenBuilder.h"
+
+FAudioPropertiesDetailArrayBuilder::FAudioPropertiesDetailArrayBuilder(const TSharedRef<IPropertyHandle>& InBaseProperty, const TSharedRef<IPropertyHandle>& InPropertySheetHandle, TWeakObjectPtr<UObject> InPropertyOwner, bool InGenerateHeader, bool InDisplayResetToDefault, bool InDisplayElementNum)
+	: FDetailArrayBuilder(InBaseProperty, InGenerateHeader, InDisplayResetToDefault, InDisplayElementNum)
+	, PropertySheetPropertyHandle(InPropertySheetHandle)
+	, PropertyOwnerObject(InPropertyOwner)
+{}
+
+void FAudioPropertiesDetailArrayBuilder::GenerateHeaderRowContent(FDetailWidgetRow& NodeRow)
+{
+	FDetailArrayBuilder::GenerateHeaderRowContent(NodeRow);
+	
+	TSharedPtr<IPropertyHandle> ArrayPropertyHandle = GetPropertyHandle();
+	TObjectPtr<UObject> PropertyOwner = PropertyOwnerObject.Get();
+
+	const TAttribute<bool> ShouldValueContentBeEnabled = TAttribute<bool>::CreateSP(this, &FAudioPropertiesDetailArrayBuilder::IsEnabled);
+	
+	NodeRow
+	.NameContent()
+	[
+		AudioPropertiesDetailsInjectorUtils::CreateOverriddenPropertyNameWidget(ArrayPropertyHandle, PropertySheetPropertyHandle, {PropertyOwner})
+	]
+	.ValueContent()
+	[
+		SNew(SHorizontalBox)
+		.IsEnabled(ShouldValueContentBeEnabled)
+		+ SHorizontalBox::Slot()
+		[
+			ArrayPropertyHandle->CreatePropertyValueWidget()
+		]
+	];
+
+	SetUpEditCondition(NodeRow);
+
+	if (GetShouldDisplayResetToDefault())
+	{
+		TSharedPtr<SResetToDefaultMenu> ResetToDefaultMenu;
+
+		NodeRow.ResetToDefaultContent()
+		[
+			SNew(SHorizontalBox)
+			.IsEnabled(ShouldValueContentBeEnabled)
+			+ SHorizontalBox::Slot()
+			[
+				SAssignNew(ResetToDefaultMenu, SResetToDefaultMenu)
+			]
+		];
+
+		ResetToDefaultMenu->AddProperty(GetPropertyHandle().ToSharedRef());
+	}
+}
+
+void FAudioPropertiesDetailArrayBuilder::GenerateChildContent( IDetailChildrenBuilder& ChildrenBuilder )
+{
+	TSharedPtr<IPropertyHandleArray> ArrayPropertyHandle = GetPropertyHandle()->AsArray();
+	check(ArrayPropertyHandle);
+
+	TAttribute<bool> ShouldArrayChildrenBeEnabled = TAttribute<bool>::CreateSP(this, &FAudioPropertiesDetailArrayBuilder::IsEnabled);
+	uint32 NumChildren = 0;
+	ArrayPropertyHandle->GetNumElements( NumChildren );
+
+	for( uint32 ChildIndex = 0; ChildIndex < NumChildren; ++ChildIndex )
+	{
+		TSharedRef<IPropertyHandle> ElementHandle = ArrayPropertyHandle->GetElement( ChildIndex );
+
+		IDetailPropertyRow& Row = ChildrenBuilder.AddProperty(ElementHandle);
+		Row.IsEnabled(ShouldArrayChildrenBeEnabled);
+	}
+}
+
+bool FAudioPropertiesDetailArrayBuilder::IsEnabled() const
+{
+	const TSharedPtr<IPropertyHandle> ArrayPropertyHandle = GetPropertyHandle();
+	const TObjectPtr<UObject> PropertyOwner = PropertyOwnerObject.Get();
+	const bool bEditableProperty = ArrayPropertyHandle->IsEditable();
+	const bool bHeaderRowEditConditionValue = GetHeaderRowEditConditionValue();
+	const bool bShouldParseProperty = AudioPropertiesDetailsInjectorUtils::ShouldParseProperty(PropertyOwner, ArrayPropertyHandle);
+
+	return bEditableProperty && bHeaderRowEditConditionValue && !bShouldParseProperty;
+}
+
+bool FAudioPropertiesDetailArrayBuilder::GetHeaderRowEditConditionValue() const
+{
+	bool bIsPropertyNegated = false;
+	const FBoolProperty* EditCondition = PropertyCustomizationHelpers::GetEditConditionProperty(GetPropertyHandle()->GetProperty(), bIsPropertyNegated);
+
+	if (!EditCondition)
+	{
+		return true;
+	}
+
+	const TObjectPtr<UObject> PropertyOwner = PropertyOwnerObject.Get();
+	bool EditConditionValue = PropertyOwner ? EditCondition->GetPropertyValue_InContainer(PropertyOwner) : true;
+	return bIsPropertyNegated ? !EditConditionValue : EditConditionValue;
+	
+}
+
+void FAudioPropertiesDetailArrayBuilder::OnHeaderRowEditConditionValueChanged(bool bValue) const
+{
+	bool bIsPropertyNegated = false;
+	FBoolProperty* EditCondition = PropertyCustomizationHelpers::GetEditConditionProperty(GetPropertyHandle()->GetProperty(), bIsPropertyNegated);
+
+	if (EditCondition)
+	{
+		const bool ValueToSet = bIsPropertyNegated ? !bValue : bValue;
+		EditCondition->SetPropertyValue_InContainer(PropertyOwnerObject.Get(), ValueToSet);
+	}
+}
+
+void FAudioPropertiesDetailArrayBuilder::SetUpEditCondition(FDetailWidgetRow& NodeRow) const
+{
+	const TSharedPtr<IPropertyHandle> ArrayPropertyHandle = GetPropertyHandle();
+	
+	const FString EditCondStr = ArrayPropertyHandle->GetMetaData(AudioPropertiesDetailsInjectorUtils::EditConditionPropertyMetadata);
+
+	if (EditCondStr.IsEmpty())
+	{
+		return;
+	}
+
+	bool bIsPropertyNegated = false;
+	const FBoolProperty* EditCondition = PropertyCustomizationHelpers::GetEditConditionProperty(GetPropertyHandle()->GetProperty(), bIsPropertyNegated);
+
+	if (!EditCondition)
+	{
+		return;
+	}
+	
+	if (!EditCondition->HasMetaData(AudioPropertiesDetailsInjectorUtils::InlineEditConditionTogglePropertyMetadata))
+	{
+		return;
+	}
+
+	const TAttribute<bool> NodeRowEditCondition = TAttribute<bool>::CreateSP(this, &FAudioPropertiesDetailArrayBuilder::GetHeaderRowEditConditionValue);
+	const FOnBooleanValueChanged OnEditConditionValueChanged = FOnBooleanValueChanged::CreateSP(this , &FAudioPropertiesDetailArrayBuilder::OnHeaderRowEditConditionValueChanged);
+	
+	NodeRow.EditCondition(NodeRowEditCondition, OnEditConditionValueChanged);
+
+}
\ No newline at end of file
diff --git a/Engine/Plugins/Experimental/AudioProperties/Source/AudioPropertiesEditor/Private/CustomNodeBuilders/AudioPropertiesDetailArrayBuilder.h b/Engine/Plugins/Experimental/AudioProperties/Source/AudioPropertiesEditor/Private/CustomNodeBuilders/AudioPropertiesDetailArrayBuilder.h
new file mode 100644
--- /dev/null
+++ b/Engine/Plugins/Experimental/AudioProperties/Source/AudioPropertiesEditor/Private/CustomNodeBuilders/AudioPropertiesDetailArrayBuilder.h
@@ -0,0 +1,24 @@
+﻿// Copyright Epic Games, Inc. All Rights Reserved.
+
+#pragma once
+
+#include "PropertyCustomizationHelpers.h"
+
+class FAudioPropertiesDetailArrayBuilder : public FDetailArrayBuilder, public TSharedFromThis<FAudioPropertiesDetailArrayBuilder>
+{
+public:
+	FAudioPropertiesDetailArrayBuilder(const TSharedRef<IPropertyHandle>& InBaseProperty, const TSharedRef<IPropertyHandle>& InPropertySheetHandle, TWeakObjectPtr<UObject> InPropertyOwner, bool InGenerateHeader = true, bool InDisplayResetToDefault = true, bool InDisplayElementNum = true);
+
+	virtual void GenerateHeaderRowContent(FDetailWidgetRow& NodeRow) override;
+
+private:
+	virtual void GenerateChildContent(IDetailChildrenBuilder& ChildrenBuilder) override;
+
+	bool IsEnabled() const;
+	bool GetHeaderRowEditConditionValue() const;
+	void OnHeaderRowEditConditionValueChanged(bool bValue) const;
+	void SetUpEditCondition(FDetailWidgetRow& NodeRow) const;
+	
+	TSharedRef<IPropertyHandle> PropertySheetPropertyHandle;
+	TWeakObjectPtr<UObject> PropertyOwnerObject;
+};
diff --git a/Engine/Plugins/Experimental/AudioProperties/Source/AudioPropertiesEditor/Private/CustomNodeBuilders/AudioPropertiesDetailNestedNodesBuilder.cpp b/Engine/Plugins/Experimental/AudioProperties/Source/AudioPropertiesEditor/Private/CustomNodeBuilders/AudioPropertiesDetailNestedNodesBuilder.cpp
new file mode 100644
--- /dev/null
+++ b/Engine/Plugins/Experimental/AudioProperties/Source/AudioPropertiesEditor/Private/CustomNodeBuilders/AudioPropertiesDetailNestedNodesBuilder.cpp
@@ -0,0 +1,174 @@
+﻿// Copyright Epic Games, Inc. All Rights Reserved.
+
+#include "AudioPropertiesDetailNestedNodesBuilder.h"
+
+#include "AudioPropertiesDetailsInjectorUtils.h"
+#include "DetailWidgetRow.h"
+#include "IDetailChildrenBuilder.h"
+#include "PropertyCustomizationHelpers.h"
+#include "SResetToDefaultMenu.h"
+
+FAudioPropertiesDetailNestedNodesBuilder::FAudioPropertiesDetailNestedNodesBuilder(TSharedRef<IPropertyHandle> InBaseProperty, const TSharedRef<IPropertyHandle>& InPropertySheetHandle, TWeakObjectPtr<UObject> InPropertyOwner, bool InGenerateHeader, bool InDisplayResetToDefault)
+	: BasePropertyHandle( InBaseProperty )
+	, PropertySheetPropertyHandle(InPropertySheetHandle)
+	, PropertyOwnerObject(InPropertyOwner)
+	, bGenerateHeader( InGenerateHeader)
+	, bDisplayResetToDefault(InDisplayResetToDefault)
+{
+	uint32 NumChildren = 0;
+	BasePropertyHandle->GetNumChildren(NumChildren);
+	ensureAlwaysMsgf(NumChildren > 0, TEXT("Customizing property handle with no children"));
+	
+	BasePropertyHandle->MarkHiddenByCustomization();
+}
+
+void FAudioPropertiesDetailNestedNodesBuilder::SetDisplayName(const FText& InDisplayName)
+{
+	DisplayName = InDisplayName;
+}
+
+FName FAudioPropertiesDetailNestedNodesBuilder::GetName() const
+{
+	return BasePropertyHandle->GetProperty()->GetFName();
+}
+
+void FAudioPropertiesDetailNestedNodesBuilder::GenerateHeaderRowContent(FDetailWidgetRow& NodeRow)
+{
+	if (bGenerateHeader)
+	{
+		const TAttribute<bool> ShouldValueContentBeEnabled = TAttribute<bool>::CreateSP(this, &FAudioPropertiesDetailNestedNodesBuilder::IsEnabled);
+		
+		FUIAction CopyAction;
+		FUIAction PasteAction;
+		BasePropertyHandle->CreateDefaultPropertyCopyPasteActions(CopyAction, PasteAction);
+		TObjectPtr<UObject> PropertyOwner = PropertyOwnerObject.Get();
+
+		NodeRow
+		.FilterString(!DisplayName.IsEmpty() ? DisplayName : BasePropertyHandle->GetPropertyDisplayName())
+		.NameContent()
+		[
+			AudioPropertiesDetailsInjectorUtils::CreateOverriddenPropertyNameWidget(BasePropertyHandle, PropertySheetPropertyHandle, {PropertyOwner})
+		]
+		.ValueContent()
+		[
+			SNew(SHorizontalBox)
+			.IsEnabled(ShouldValueContentBeEnabled)
+			+ SHorizontalBox::Slot()
+			[
+				BasePropertyHandle->CreatePropertyValueWidget()
+			]
+		]
+		.CopyAction(CopyAction)
+		.PasteAction(PasteAction);
+
+		if (bDisplayResetToDefault)
+		{
+			TSharedPtr<SResetToDefaultMenu> ResetToDefaultMenu;
+
+			NodeRow.ResetToDefaultContent()
+			[
+				SNew(SHorizontalBox)
+				.IsEnabled(ShouldValueContentBeEnabled)
+				+ SHorizontalBox::Slot()
+				[
+					SAssignNew(ResetToDefaultMenu, SResetToDefaultMenu)
+				]
+			];
+
+			ResetToDefaultMenu->AddProperty(BasePropertyHandle);
+		}
+
+		SetUpEditCondition(NodeRow);
+	}
+}
+
+void FAudioPropertiesDetailNestedNodesBuilder::GenerateChildContent(IDetailChildrenBuilder& ChildrenBuilder)
+{
+	uint32 NumChildren = 0;
+	BasePropertyHandle->GetNumChildren( NumChildren );
+
+	TAttribute<bool> ShouldSetChildrenBeEnabled = TAttribute<bool>::CreateSP(this, &FAudioPropertiesDetailNestedNodesBuilder::IsEnabled);
+	
+	for( uint32 ChildIndex = 0; ChildIndex < NumChildren; ++ChildIndex )
+	{
+		TSharedPtr<IPropertyHandle> ElementHandle = BasePropertyHandle->GetChildHandle(ChildIndex);
+
+		if (ElementHandle.IsValid())
+		{
+			IDetailPropertyRow& Row = ChildrenBuilder.AddProperty(ElementHandle.ToSharedRef());
+			Row.IsEnabled(ShouldSetChildrenBeEnabled);
+		}
+	}
+}
+
+TSharedPtr<IPropertyHandle> FAudioPropertiesDetailNestedNodesBuilder::GetPropertyHandle() const
+{
+	return BasePropertyHandle;
+}
+
+bool FAudioPropertiesDetailNestedNodesBuilder::IsEnabled() const
+{
+	TObjectPtr<UObject> PropertyOwner = PropertyOwnerObject.Get();
+	const bool bEditableProperty = BasePropertyHandle->IsEditable();
+	const bool bHeaderRowEditConditionValue = GetHeaderRowEditConditionValue();
+	const bool bShouldParseProperty = AudioPropertiesDetailsInjectorUtils::ShouldParseProperty(PropertyOwner, BasePropertyHandle);
+
+	return bEditableProperty && bHeaderRowEditConditionValue && !bShouldParseProperty;
+}
+
+bool FAudioPropertiesDetailNestedNodesBuilder::GetHeaderRowEditConditionValue() const
+{
+	bool bIsPropertyNegated = false;
+	FBoolProperty* EditCondition = PropertyCustomizationHelpers::GetEditConditionProperty(GetPropertyHandle()->GetProperty(), bIsPropertyNegated);
+
+	if (!EditCondition)
+	{
+		return true;
+	}
+
+	bool EditConditionValue = PropertyOwnerObject.Get() ? EditCondition->GetPropertyValue_InContainer(PropertyOwnerObject.Get()) : true;
+	return bIsPropertyNegated ? !EditConditionValue : EditConditionValue;
+	
+}
+
+void FAudioPropertiesDetailNestedNodesBuilder::OnHeaderRowEditConditionValueChanged(bool bValue) const
+{
+	bool bIsPropertyNegated = false;
+	FBoolProperty* EditCondition = PropertyCustomizationHelpers::GetEditConditionProperty(GetPropertyHandle()->GetProperty(), bIsPropertyNegated);
+
+	if (EditCondition)
+	{
+		const bool ValueToSet = bIsPropertyNegated ? !bValue : bValue;
+		EditCondition->SetPropertyValue_InContainer(PropertyOwnerObject.Get(), ValueToSet);
+	}
+}
+
+void FAudioPropertiesDetailNestedNodesBuilder::SetUpEditCondition(FDetailWidgetRow& NodeRow) const
+{
+	TSharedPtr<IPropertyHandle> ArrayPropertyHandle = GetPropertyHandle();
+	
+	const FString EditCondStr = ArrayPropertyHandle->GetMetaData(AudioPropertiesDetailsInjectorUtils::EditConditionPropertyMetadata);
+
+	if (EditCondStr.IsEmpty())
+	{
+		return;
+	}
+
+	bool bIsPropertyNegated = false;
+	FBoolProperty* EditCondition = PropertyCustomizationHelpers::GetEditConditionProperty(GetPropertyHandle()->GetProperty(), bIsPropertyNegated);
+
+	if (!EditCondition)
+	{
+		return;
+	}
+	
+	if (!EditCondition->HasMetaData(AudioPropertiesDetailsInjectorUtils::InlineEditConditionTogglePropertyMetadata))
+	{
+		return;
+	}
+
+	const TAttribute<bool> NodeRowEditCondition = TAttribute<bool>::CreateSP(this, &FAudioPropertiesDetailNestedNodesBuilder::GetHeaderRowEditConditionValue);
+	const FOnBooleanValueChanged OnEditConditionValueChanged = FOnBooleanValueChanged::CreateSP(this , &FAudioPropertiesDetailNestedNodesBuilder::OnHeaderRowEditConditionValueChanged);
+	
+	NodeRow.EditCondition(NodeRowEditCondition, OnEditConditionValueChanged);
+}
diff --git a/Engine/Plugins/Experimental/AudioProperties/Source/AudioPropertiesEditor/Private/CustomNodeBuilders/AudioPropertiesDetailNestedNodesBuilder.h b/Engine/Plugins/Experimental/AudioProperties/Source/AudioPropertiesEditor/Private/CustomNodeBuilders/AudioPropertiesDetailNestedNodesBuilder.h
new file mode 100644
--- /dev/null
+++ b/Engine/Plugins/Experimental/AudioProperties/Source/AudioPropertiesEditor/Private/CustomNodeBuilders/AudioPropertiesDetailNestedNodesBuilder.h
@@ -0,0 +1,40 @@
+﻿// Copyright Epic Games, Inc. All Rights Reserved.
+
+#pragma once
+
+#include "IDetailCustomNodeBuilder.h"
+#include "PropertyHandle.h"
+
+class FAudioPropertiesDetailNestedNodesBuilder : public IDetailCustomNodeBuilder , public TSharedFromThis<FAudioPropertiesDetailNestedNodesBuilder>
+{
+public:
+	FAudioPropertiesDetailNestedNodesBuilder(TSharedRef<IPropertyHandle> InBaseProperty, const TSharedRef<IPropertyHandle>& InPropertySheetHandle, TWeakObjectPtr<UObject> InPropertyOwner, bool InGenerateHeader = true, bool InDisplayResetToDefault = true);
+
+	virtual bool RequiresTick() const override { return false; }
+	virtual void Tick( float DeltaTime ) override {}
+	
+	void SetDisplayName( const FText& InDisplayName );
+	virtual FName GetName() const override;
+
+	virtual bool InitiallyCollapsed() const override { return false; }
+
+	virtual void GenerateHeaderRowContent(FDetailWidgetRow& NodeRow) override;
+	virtual void GenerateChildContent(IDetailChildrenBuilder& ChildrenBuilder) override;
+	
+	virtual TSharedPtr<IPropertyHandle> GetPropertyHandle() const override;
+
+private:
+	bool IsEnabled() const;
+	bool GetHeaderRowEditConditionValue() const;
+	void OnHeaderRowEditConditionValueChanged(bool bValue) const;
+	void SetUpEditCondition(FDetailWidgetRow& NodeRow) const;
+	
+	FText DisplayName;
+	TSharedRef<IPropertyHandle> BasePropertyHandle;
+	TSharedRef<IPropertyHandle> PropertySheetPropertyHandle;
+	TWeakObjectPtr<UObject> PropertyOwnerObject;
+
+	bool bGenerateHeader;
+	bool bDisplayResetToDefault;
+	
+};
diff --git a/Engine/Plugins/Experimental/AudioProperties/Source/AudioPropertiesEditor/Private/CustomNodeBuilders/AudioPropertiesDetailSetBuilder.cpp b/Engine/Plugins/Experimental/AudioProperties/Source/AudioPropertiesEditor/Private/CustomNodeBuilders/AudioPropertiesDetailSetBuilder.cpp
new file mode 100644
--- /dev/null
+++ b/Engine/Plugins/Experimental/AudioProperties/Source/AudioPropertiesEditor/Private/CustomNodeBuilders/AudioPropertiesDetailSetBuilder.cpp
@@ -0,0 +1,192 @@
+﻿// Copyright Epic Games, Inc. All Rights Reserved.
+
+#include "AudioPropertiesDetailSetBuilder.h"
+#include "AudioPropertiesDetailsInjectorUtils.h"
+#include "DetailWidgetRow.h"
+#include "IDetailChildrenBuilder.h"
+#include "PropertyCustomizationHelpers.h"
+#include "SResetToDefaultMenu.h"
+	
+FAudioPropertiesDetailSetBuilder::FAudioPropertiesDetailSetBuilder(TSharedRef<IPropertyHandle> InBaseProperty, const TSharedRef<IPropertyHandle>& InPropertySheetHandle, TWeakObjectPtr<UObject> InPropertyOwner, bool InGenerateHeader, bool InDisplayResetToDefault, bool InDisplayElementNum)
+	: SetPropertyHandle( InBaseProperty->AsSet() )
+	, BasePropertyHandle( InBaseProperty )
+	, PropertySheetPropertyHandle(InPropertySheetHandle)
+	, PropertyOwnerObject(InPropertyOwner)
+	, bGenerateHeader( InGenerateHeader)
+	, bDisplayResetToDefault(InDisplayResetToDefault)
+	, bDisplayElementNum(InDisplayElementNum)
+{
+	check( SetPropertyHandle.IsValid() );
+
+	FSimpleDelegate OnNumChildrenChanged = FSimpleDelegate::CreateRaw( this, &FAudioPropertiesDetailSetBuilder::OnNumChildrenChanged );
+	OnNumElementsChangedHandle = SetPropertyHandle->SetOnNumElementsChanged( OnNumChildrenChanged );
+	
+	BasePropertyHandle->MarkHiddenByCustomization();
+}
+
+FAudioPropertiesDetailSetBuilder::~FAudioPropertiesDetailSetBuilder()
+{
+	SetPropertyHandle->UnregisterOnNumElementsChanged(OnNumElementsChangedHandle);
+}
+
+void FAudioPropertiesDetailSetBuilder::SetDisplayName(const FText& InDisplayName)
+{
+	DisplayName = InDisplayName;
+}
+
+FName FAudioPropertiesDetailSetBuilder::GetName() const
+{
+	return BasePropertyHandle->GetProperty()->GetFName();
+}
+
+void FAudioPropertiesDetailSetBuilder::GenerateHeaderRowContent(FDetailWidgetRow& NodeRow)
+{
+	if (bGenerateHeader)
+	{
+		const TAttribute<bool> ShouldValueContentBeEnabled = TAttribute<bool>::CreateSP<>(this, &FAudioPropertiesDetailSetBuilder::IsEnabled);
+
+		TSharedPtr<SHorizontalBox> ContentHorizontalBox;
+		SAssignNew(ContentHorizontalBox, SHorizontalBox);
+		if (bDisplayElementNum)
+		{
+			ContentHorizontalBox->AddSlot()
+			[
+				BasePropertyHandle->CreatePropertyValueWidget()
+			];
+		}
+
+		FUIAction CopyAction;
+		FUIAction PasteAction;
+		BasePropertyHandle->CreateDefaultPropertyCopyPasteActions(CopyAction, PasteAction);
+		TObjectPtr<UObject> PropertyOwner = PropertyOwnerObject.Get();
+
+		NodeRow
+		.FilterString(!DisplayName.IsEmpty() ? DisplayName : BasePropertyHandle->GetPropertyDisplayName())
+		.NameContent()
+		[
+			AudioPropertiesDetailsInjectorUtils::CreateOverriddenPropertyNameWidget(BasePropertyHandle, PropertySheetPropertyHandle, {PropertyOwner})
+		]
+		.ValueContent()
+		[
+			SNew(SHorizontalBox)
+			.IsEnabled(ShouldValueContentBeEnabled)
+			+ SHorizontalBox::Slot()
+			[
+				BasePropertyHandle->CreatePropertyValueWidget()
+			]
+		]
+		.CopyAction(CopyAction)
+		.PasteAction(PasteAction);
+
+		if (bDisplayResetToDefault)
+		{
+			TSharedPtr<SResetToDefaultMenu> ResetToDefaultMenu;
+
+			NodeRow.ResetToDefaultContent()
+			[
+				SNew(SHorizontalBox)
+				.IsEnabled(ShouldValueContentBeEnabled)
+				+ SHorizontalBox::Slot()
+				[
+					SAssignNew(ResetToDefaultMenu, SResetToDefaultMenu)
+				]
+			];
+
+			ResetToDefaultMenu->AddProperty(BasePropertyHandle);
+		}
+
+		SetUpEditCondition(NodeRow);
+	}
+}
+
+void FAudioPropertiesDetailSetBuilder::GenerateChildContent(IDetailChildrenBuilder& ChildrenBuilder)
+{
+	uint32 NumChildren = 0;
+	SetPropertyHandle->GetNumElements( NumChildren );
+
+	TAttribute<bool> ShouldSetChildrenBeEnabled = TAttribute<bool>::CreateSP(this, &FAudioPropertiesDetailSetBuilder::IsEnabled);
+	
+	for( uint32 ChildIndex = 0; ChildIndex < NumChildren; ++ChildIndex )
+	{
+		TSharedRef<IPropertyHandle> ElementHandle = SetPropertyHandle->GetElement( ChildIndex );
+		IDetailPropertyRow& Row = ChildrenBuilder.AddProperty(ElementHandle);
+		Row.IsEnabled(ShouldSetChildrenBeEnabled);
+	}
+}
+
+TSharedPtr<IPropertyHandle> FAudioPropertiesDetailSetBuilder::GetPropertyHandle() const
+{
+	return BasePropertyHandle;
+}
+
+void FAudioPropertiesDetailSetBuilder::OnNumChildrenChanged() const
+{
+	OnRebuildChildren.ExecuteIfBound();
+}
+
+bool FAudioPropertiesDetailSetBuilder::IsEnabled() const
+{
+	TObjectPtr<UObject> PropertyOwner = PropertyOwnerObject.Get();
+	const bool bEditableProperty = BasePropertyHandle->IsEditable();
+	const bool bHeaderRowEditConditionValue = GetHeaderRowEditConditionValue();
+	const bool bShouldParseProperty = AudioPropertiesDetailsInjectorUtils::ShouldParseProperty(PropertyOwner, BasePropertyHandle);
+
+	return bEditableProperty && bHeaderRowEditConditionValue && !bShouldParseProperty;
+}
+
+bool FAudioPropertiesDetailSetBuilder::GetHeaderRowEditConditionValue() const
+{
+	bool bIsPropertyNegated = false;
+	FBoolProperty* EditCondition = PropertyCustomizationHelpers::GetEditConditionProperty(GetPropertyHandle()->GetProperty(), bIsPropertyNegated);
+
+	if (!EditCondition)
+	{
+		return true;
+	}
+
+	bool EditConditionValue = PropertyOwnerObject.Get() ? EditCondition->GetPropertyValue_InContainer(PropertyOwnerObject.Get()) : true;
+	return bIsPropertyNegated ? !EditConditionValue : EditConditionValue;
+	
+}
+
+void FAudioPropertiesDetailSetBuilder::OnHeaderRowEditConditionValueChanged(bool bValue) const
+{
+	bool bIsPropertyNegated = false;
+	FBoolProperty* EditCondition = PropertyCustomizationHelpers::GetEditConditionProperty(GetPropertyHandle()->GetProperty(), bIsPropertyNegated);
+
+	if (EditCondition)
+	{
+		const bool ValueToSet = bIsPropertyNegated ? !bValue : bValue;
+		EditCondition->SetPropertyValue_InContainer(PropertyOwnerObject.Get(), ValueToSet);
+	}
+}
+
+void FAudioPropertiesDetailSetBuilder::SetUpEditCondition(FDetailWidgetRow& NodeRow) const
+{
+	TSharedPtr<IPropertyHandle> ArrayPropertyHandle = GetPropertyHandle();
+	
+	const FString EditCondStr = ArrayPropertyHandle->GetMetaData(AudioPropertiesDetailsInjectorUtils::EditConditionPropertyMetadata);
+
+	if (EditCondStr.IsEmpty())
+	{
+		return;
+	}
+
+	bool bIsPropertyNegated = false;
+	FBoolProperty* EditCondition = PropertyCustomizationHelpers::GetEditConditionProperty(GetPropertyHandle()->GetProperty(), bIsPropertyNegated);
+
+	if (!EditCondition)
+	{
+		return;
+	}
+	
+	if (!EditCondition->HasMetaData(AudioPropertiesDetailsInjectorUtils::InlineEditConditionTogglePropertyMetadata))
+	{
+		return;
+	}
+
+	const TAttribute<bool> NodeRowEditCondition = TAttribute<bool>::CreateSP(this, &FAudioPropertiesDetailSetBuilder::GetHeaderRowEditConditionValue);
+	const FOnBooleanValueChanged OnEditConditionValueChanged = FOnBooleanValueChanged::CreateSP(this , &FAudioPropertiesDetailSetBuilder::OnHeaderRowEditConditionValueChanged);
+	
+	NodeRow.EditCondition(NodeRowEditCondition, OnEditConditionValueChanged);
+}
\ No newline at end of file
diff --git a/Engine/Plugins/Experimental/AudioProperties/Source/AudioPropertiesEditor/Private/CustomNodeBuilders/AudioPropertiesDetailSetBuilder.h b/Engine/Plugins/Experimental/AudioProperties/Source/AudioPropertiesEditor/Private/CustomNodeBuilders/AudioPropertiesDetailSetBuilder.h
new file mode 100644
--- /dev/null
+++ b/Engine/Plugins/Experimental/AudioProperties/Source/AudioPropertiesEditor/Private/CustomNodeBuilders/AudioPropertiesDetailSetBuilder.h
@@ -0,0 +1,49 @@
+﻿// Copyright Epic Games, Inc. All Rights Reserved.
+
+#pragma once
+
+#include "IDetailCustomNodeBuilder.h"
+#include "PropertyHandle.h"
+
+class FAudioPropertiesDetailSetBuilder : public IDetailCustomNodeBuilder , public TSharedFromThis<FAudioPropertiesDetailSetBuilder>
+{
+public:
+	FAudioPropertiesDetailSetBuilder(TSharedRef<IPropertyHandle> InBaseProperty, const TSharedRef<IPropertyHandle>& InPropertySheetHandle, TWeakObjectPtr<UObject> InPropertyOwner, bool InGenerateHeader = true, bool InDisplayResetToDefault = true, bool InDisplayElementNum = true);
+	virtual ~FAudioPropertiesDetailSetBuilder();
+
+	virtual bool RequiresTick() const override { return false; }
+	virtual void Tick( float DeltaTime ) override {}
+
+	void SetDisplayName( const FText& InDisplayName );
+	virtual FName GetName() const override;
+
+	virtual bool InitiallyCollapsed() const override { return false; }
+
+	virtual void GenerateHeaderRowContent(FDetailWidgetRow& NodeRow) override;
+	virtual void GenerateChildContent(IDetailChildrenBuilder& ChildrenBuilder) override;
+	
+	virtual TSharedPtr<IPropertyHandle> GetPropertyHandle() const override;
+
+protected:
+	virtual void SetOnRebuildChildren(FSimpleDelegate InOnRebuildChildren) override { OnRebuildChildren = InOnRebuildChildren; }
+	void OnNumChildrenChanged() const;
+
+private:
+	bool IsEnabled() const;
+	bool GetHeaderRowEditConditionValue() const;
+	void OnHeaderRowEditConditionValueChanged(bool bValue) const;
+	void SetUpEditCondition(FDetailWidgetRow& NodeRow) const;
+	
+	FText DisplayName;
+	TSharedPtr<IPropertyHandleSet> SetPropertyHandle;
+	TSharedRef<IPropertyHandle> BasePropertyHandle;
+	TSharedRef<IPropertyHandle> PropertySheetPropertyHandle;
+	TWeakObjectPtr<UObject> PropertyOwnerObject;
+	
+	FSimpleDelegate OnRebuildChildren;
+	FDelegateHandle OnNumElementsChangedHandle;
+	
+	bool bGenerateHeader;
+	bool bDisplayResetToDefault;
+	bool bDisplayElementNum;
+};
diff --git a/Engine/Plugins/Experimental/AudioProperties/Source/AudioPropertiesEditor/Private/InstancedAudioPropertyBagDetails.cpp b/Engine/Plugins/Experimental/AudioProperties/Source/AudioPropertiesEditor/Private/InstancedAudioPropertyBagDetails.cpp
new file mode 100644
--- /dev/null
+++ b/Engine/Plugins/Experimental/AudioProperties/Source/AudioPropertiesEditor/Private/InstancedAudioPropertyBagDetails.cpp
@@ -0,0 +1,227 @@
+// Copyright Epic Games, Inc. All Rights Reserved.
+
+#include "InstancedAudioPropertyBagDetails.h"
+
+#include "AudioPropertiesSheet.h"
+#include "DetailLayoutBuilder.h"
+#include "DetailWidgetRow.h"
+#include "IPropertyUtilities.h"
+#include "PropertyBagDetails.h"
+#include "ScopedTransaction.h"
+#include "Widgets/Text/STextBlock.h"
+
+#define LOCTEXT_NAMESPACE "InstancedAudioPropertyBagDetails.h"
+
+namespace InstancedAudioPropertyBagDetailsCVars
+{
+	static bool ShowIDsOnAudioPropertyBagDetailsCVar = 0;
+	FAutoConsoleVariableRef CVarAudioPropertyBagDetailsShowIDs(
+		TEXT("au.AudioProperties.ShowPropertiesIDs"),
+		ShowIDsOnAudioPropertyBagDetailsCVar,
+		TEXT("Show properties IDs on audio property sheet details.\n")
+		TEXT("0: Disabled (default), 1: Enabled"),
+		ECVF_Default);
+}
+
+namespace InstancedAudioPropertyBagDetailsPrivate
+{
+	FGuid GetPropertyID(const FAudioPropertiesSheet& InLeafMostPropertySheet, TSharedPtr<IPropertyHandle> InProperty)
+	{
+		if (!InProperty)
+		{
+			return FGuid();
+		}
+
+		auto FindPropertyID = [](const FAudioPropertiesSheet& PropertySheet, const FName& PropertyName) -> FGuid
+		{
+			if (const FPropertyBagPropertyDesc* PropertyDesc = PropertySheet.Properties.FindPropertyDescByName(PropertyName))
+			{
+				return PropertyDesc->ID;
+			}
+			return FGuid(); 
+		};
+
+		const FName PropertyName = InProperty->GetProperty()->GetFName();
+
+		// Find the closest parent sheet and check for the property ID
+		if (TObjectPtr<const UAudioPropertiesSheetAsset> OwningSheet = InLeafMostPropertySheet.FindClosestParentWithProperty(*InProperty->GetProperty()).Value)
+		{
+			return InLeafMostPropertySheet.IsPropertyOverridden(*InProperty->GetProperty()) ? FindPropertyID(InLeafMostPropertySheet, PropertyName) : FindPropertyID(OwningSheet->PropertiesSheet, PropertyName);
+		}
+
+		//in case we did not find a parent, this property is likely to be local
+		return FindPropertyID(InLeafMostPropertySheet, PropertyName);
+	}
+}
+
+FInstancedAudioPropertyBagDetails::FInstancedAudioPropertyBagDetails(TSharedPtr<IPropertyHandle> InInstancedPropertyBagHandle, const TSharedPtr<IPropertyUtilities>& InPropUtils, const UAudioPropertiesSheetAsset* InOwningAudioPropertiesSheetAsset, FAudioPropertiesSheet* InLeafMostPropertySheet, bool bInIsLeafBag /*= false*/)
+ 	: FPropertyBagInstanceDataDetails(InInstancedPropertyBagHandle, InPropUtils, /*bInFixedLayout*/ !bInIsLeafBag, /*bInAllowArrays*/ true)
+ 	, InstancedPropertyBagHandle(InInstancedPropertyBagHandle)
+    , OwningAudioPropertiesSheetAsset(InOwningAudioPropertiesSheetAsset)
+    , LeafMostPropertySheet(InLeafMostPropertySheet)
+	, CachedPropUtils(InPropUtils)
+ 	, bIsLeafBag(bInIsLeafBag)
+{
+	check(InstancedPropertyBagHandle.IsValid())
+	check(LeafMostPropertySheet)
+}
+
+void FInstancedAudioPropertyBagDetails::OnChildRowAdded(IDetailPropertyRow& ChildRow)
+{ 
+	if (bIsLeafBag)
+ 	{
+		DrawLeafBagChildRow(ChildRow);
+		return;
+ 	}
+
+	DrawInheritedBagChildRow(ChildRow);
+}
+
+void FInstancedAudioPropertyBagDetails::DrawLeafBagChildRow(IDetailPropertyRow& ChildRow)
+{
+	TSharedPtr<IPropertyHandle> ChildPropertyHandle = ChildRow.GetPropertyHandle();
+	check(ChildPropertyHandle->GetProperty())
+
+	FPropertyBagInstanceDataDetails::OnChildRowAdded(ChildRow);
+
+	TSharedRef<SWidget> PropertyBagNameWidget = ChildRow.CustomNameWidget()->Widget;
+	TSharedRef<SWidget> PropertyBagValueWidget = ChildRow.CustomValueWidget()->Widget;
+
+	ChildRow.CustomWidget(true)
+		.NameContent()
+		[
+			PropertyBagNameWidget
+		]
+		.ValueContent()
+		[
+			PropertyBagValueWidget
+		]
+		.ExtensionContent()
+		[
+			SNew(STextBlock).Text_Lambda([this, ChildPropertyHandle]()
+				{
+					return FText::FromString(InstancedAudioPropertyBagDetailsPrivate::GetPropertyID(*LeafMostPropertySheet, ChildPropertyHandle).ToString());
+				})
+			.Visibility(InstancedAudioPropertyBagDetailsCVars::ShowIDsOnAudioPropertyBagDetailsCVar ? EVisibility::Visible : EVisibility::Hidden)
+		];
+
+	TAttribute<bool> bIsInheritedProperty = LeafMostPropertySheet->IsPropertyOverridden(*ChildPropertyHandle->GetProperty());
+
+	if (bIsInheritedProperty.Get() == true)
+	{
+		FOnBooleanValueChanged OnOverrideChanged = FOnBooleanValueChanged::CreateLambda([this, ChildPropertyHandle](const bool bOverrideValue)
+		{
+			HandlePropertyOverrideChanged(ChildPropertyHandle, bOverrideValue);
+		});
+	
+		ChildRow.EditCondition(bIsInheritedProperty, OnOverrideChanged);
+	}
+}
+
+void FInstancedAudioPropertyBagDetails::DrawInheritedBagChildRow(IDetailPropertyRow& ChildRow)
+{
+	TSharedPtr<IPropertyHandle> ChildPropertyHandle = ChildRow.GetPropertyHandle();
+	
+	check(ChildPropertyHandle->GetProperty())
+	check(InstancedPropertyBagHandle.IsValid());
+	check(OwningAudioPropertiesSheetAsset);
+	check(LeafMostPropertySheet);
+
+	TSharedPtr<SWidget> NameWidget;
+	TSharedPtr<SWidget> ValueWidget;
+	FDetailWidgetRow Row;
+
+	ChildRow.GetDefaultWidgets(NameWidget, ValueWidget, Row);
+ 
+	const bool bIsPropertyOverridenByLeafMost = LeafMostPropertySheet->IsPropertyOverridden(*ChildPropertyHandle->GetProperty());
+	const bool bIsPropertyOverriddenByChild = LeafMostPropertySheet->FindClosestParentWithProperty(*ChildPropertyHandle->GetProperty()).Value != OwningAudioPropertiesSheetAsset;
+
+	auto GenerateOwningSheetText = [this]()
+	{
+		if (!OwningAudioPropertiesSheetAsset)
+		{
+			return FText();
+		}
+
+		FString OwningAssetName;
+		OwningAudioPropertiesSheetAsset->GetName(OwningAssetName);
+		const FString OwningAssetString = FString::Printf(TEXT("[%s] "), *OwningAssetName);
+
+		return FText::FromString(OwningAssetString);
+	};
+
+	ChildRow
+		.CustomWidget(/*bShowChildren*/true)
+		.NameContent()
+		[
+			SNew(SHorizontalBox)
+			// Name
+			+ SHorizontalBox::Slot()
+			.AutoWidth()
+			.VAlign(EVerticalAlignment::VAlign_Center)
+			[
+				SNew(SHorizontalBox)
+				+ SHorizontalBox::Slot()
+				.AutoWidth()
+				[
+					SNew(STextBlock).Text_Lambda(GenerateOwningSheetText)
+					.Font(IDetailLayoutBuilder::GetDetailFontItalic())
+				]
+				+ SHorizontalBox::Slot()
+				.AutoWidth()
+				[
+					SNew(STextBlock)
+					.Font(IDetailLayoutBuilder::GetDetailFont())
+					.Text(FText::FromName(ChildPropertyHandle->GetProperty()->GetFName()))
+				]
+			]
+		]
+		.ValueContent()
+		[
+			SNew(SHorizontalBox)
+			+ SHorizontalBox::Slot()
+			[
+				ValueWidget.ToSharedRef()
+			]
+		]
+		.ExtensionContent()
+		[
+			SNew(STextBlock).Text_Lambda([this, ChildPropertyHandle]()
+			{
+				return FText::FromString(InstancedAudioPropertyBagDetailsPrivate::GetPropertyID(*LeafMostPropertySheet, ChildPropertyHandle).ToString());
+			})
+			.Visibility(InstancedAudioPropertyBagDetailsCVars::ShowIDsOnAudioPropertyBagDetailsCVar ? EVisibility::Visible : EVisibility::Hidden)
+		];
+
+	TAttribute<bool> bIsPropertyOverridden = bIsPropertyOverridenByLeafMost || bIsPropertyOverriddenByChild;
+	ChildRow.Visibility(bIsPropertyOverridden.Get() ? EVisibility::Hidden : EVisibility::Visible);
+
+	FOnBooleanValueChanged OnOverrideChanged = FOnBooleanValueChanged::CreateLambda([this, ChildPropertyHandle](const bool bOverrideValue)
+	{
+		HandlePropertyOverrideChanged(ChildPropertyHandle, bOverrideValue);
+	});
+
+	ChildRow.EditCondition(bIsPropertyOverridden, OnOverrideChanged);
+}
+
+void FInstancedAudioPropertyBagDetails::HandlePropertyOverrideChanged(TSharedPtr<IPropertyHandle> PropertyThatChanged, const bool bOverrideState)
+{
+	FText TransactionText = FText::Format(LOCTEXT("OnCheckStateChanged", "Changed Inheritance for property {0} to {1}"),
+		FText::FromName(PropertyThatChanged->GetProperty()->GetFName()),
+		bOverrideState ? FText::FromString(TEXT("Overridden")) : FText::FromString(TEXT("Inherited")));
+
+	FScopedTransaction Transaction(TransactionText);
+
+	InstancedPropertyBagHandle->NotifyPreChange();
+	LeafMostPropertySheet->UpdatePropertyOverride(*PropertyThatChanged->GetProperty(), bOverrideState);
+	InstancedPropertyBagHandle->NotifyPostChange(EPropertyChangeType::ValueSet);
+	InstancedPropertyBagHandle->NotifyFinishedChangingProperties();
+	InstancedPropertyBagHandle->RequestRebuildChildren();
+
+	if (CachedPropUtils)
+	{
+		CachedPropUtils->ForceRefresh();
+	}
+}
+
+#undef LOCTEXT_NAMESPACE
\ No newline at end of file
diff --git a/Engine/Plugins/Experimental/AudioProperties/Source/AudioPropertiesEditor/Private/SAudioPropertiesSheetBuilderWidget.cpp b/Engine/Plugins/Experimental/AudioProperties/Source/AudioPropertiesEditor/Private/SAudioPropertiesSheetBuilderWidget.cpp
new file mode 100644
--- /dev/null
+++ b/Engine/Plugins/Experimental/AudioProperties/Source/AudioPropertiesEditor/Private/SAudioPropertiesSheetBuilderWidget.cpp
@@ -0,0 +1,746 @@
+// Copyright Epic Games, Inc. All Rights Reserved.
+
+#include "SAudioPropertiesSheetBuilderWidget.h"
+
+#include "AssetRegistry/AssetRegistryModule.h"
+#include "AudioPropertiesSheet.h"
+#include "AudioPropertiesUtils.h"
+#include "Editor.h"
+#include "Framework/Application/SlateApplication.h"
+#include "Misc/MessageDialog.h"
+#include "PropertyCustomizationHelpers.h"
+#include "Sound/SoundBase.h"
+#include "UObject/UnrealType.h"
+#include "Widgets/Input/SButton.h"
+#include "Widgets/Input/SComboButton.h"
+#include "Widgets/Input/SCheckBox.h"
+#include "Widgets/Input/SEditableTextBox.h"
+#include "Widgets/Layout/SScrollBox.h"
+#include "Widgets/Text/STextBlock.h"
+#include "Widgets/Views/STableViewBase.h"
+
+#define LOCTEXT_NAMESPACE "SAudioPropertiesSheetBuilderWidget"
+
+namespace AudioPropertiesSheetBuilderWidget
+{
+	static const FName AddPropertyColumnName = "AddProperty";
+	static const FName CopyPropertyValueColumnName = "CopyPropertyValue";
+	static const FName PropertyNameColumnName = "PropertyName";
+	static const FName PropertyValueColumnName = "PropertyValue";
+
+	const float ColumnItemBoxPadding = 1.f;
+	const FVector2D HeaderItemPadding = FVector2D(5.f, 3.f);
+	const FVector2D ComboBoxPadding = FVector2D(5.f, 0.f);
+
+	struct FAssetPickerParams
+	{
+		FAssetData InitialObject = FAssetData();
+		bool bAllowClear = false;
+		TArray<const UClass*> AllowedClasses;
+		FOnAssetSelected OnAssetSelected;
+	};
+
+	TSharedRef<SWidget> BuildAssetPickerWidget(const FAssetPickerParams& InAssetPickerParams)
+	{
+		return PropertyCustomizationHelpers::MakeAssetPickerWithMenu(
+			InAssetPickerParams.InitialObject,
+			InAssetPickerParams.bAllowClear,
+			InAssetPickerParams.AllowedClasses,
+			PropertyCustomizationHelpers::GetNewAssetFactoriesForClasses(InAssetPickerParams.AllowedClasses),
+			FOnShouldFilterAsset(),
+			InAssetPickerParams.OnAssetSelected,
+			FSimpleDelegate());
+	}
+		
+	class SPropertyRow : public SMultiColumnTableRow<AudioPropertiesSheetAssetBuilder::FPropertyRequestPtr>
+	{
+	public:
+		SLATE_BEGIN_ARGS(SPropertyRow) {}
+		SLATE_END_ARGS()
+
+		void Construct(const FArguments& InArgs, const TSharedRef<STableViewBase>& InOwnerTableView, const AudioPropertiesSheetAssetBuilder::FPropertyRequestPtr InEntryItem, const UObject* InSourceObjectPtr)
+		{
+			Item = InEntryItem;
+			SourceObjectPtr = InSourceObjectPtr;
+			SMultiColumnTableRow<AudioPropertiesSheetAssetBuilder::FPropertyRequestPtr>::Construct(FSuperRowType::FArguments(), InOwnerTableView);
+		}
+
+		virtual TSharedRef<SWidget> GenerateWidgetForColumn(const FName& ColumnName) override
+		{
+			if (ColumnName == AddPropertyColumnName)
+			{
+				return  SNew(SHorizontalBox)
+					+ SHorizontalBox::Slot()
+					.AutoWidth()
+					.Padding(ColumnItemBoxPadding)
+					[
+						SNew(SCheckBox)
+						.IsChecked_Lambda([this]()
+						{
+							return Item->bInheritProperty ? ECheckBoxState::Checked : ECheckBoxState::Unchecked;
+						})
+						.OnCheckStateChanged_Lambda([this](ECheckBoxState NewState)
+						{
+							Item->bInheritProperty = NewState == ECheckBoxState::Checked;
+							if (!Item->bInheritProperty)
+							{
+								//if we are not inheriting the property we should not inherit the value
+								Item->bInheritValue = false;
+							}
+						})
+					];
+        
+			}
+			else if (ColumnName == CopyPropertyValueColumnName)
+			{
+					return  SNew(SHorizontalBox)
+					+ SHorizontalBox::Slot()
+					.AutoWidth()
+					.Padding(ColumnItemBoxPadding)
+					[
+						SNew(SCheckBox)
+						.IsChecked_Lambda([this]()
+						{
+							return Item->bInheritValue ? ECheckBoxState::Checked : ECheckBoxState::Unchecked;
+						})
+						.OnCheckStateChanged_Lambda([this](ECheckBoxState NewState)
+						{
+							Item->bInheritValue = NewState == ECheckBoxState::Checked;
+							if (Item->bInheritValue)
+							{
+								//if we are inheriting the value we also want to inherit the property
+								Item->bInheritProperty = true;
+							}
+						})
+					];
+			}
+			else if (ColumnName == PropertyNameColumnName)
+			{
+				const FText PropertyName = Item->CachedProperty ? Item->CachedProperty->GetDisplayNameText() : FText();
+
+				return SNew(SHorizontalBox)
+					+ SHorizontalBox::Slot()
+					.AutoWidth()
+					.Padding(4.f, ColumnItemBoxPadding)
+					[
+						SNew(STextBlock)
+						.AutoWrapText(true)
+						.Text(PropertyName)
+					];
+			}
+			else if (ColumnName == PropertyValueColumnName)
+			{
+				FString PropertyValue;
+								
+				if (SourceObjectPtr && SourceObjectPtr->IsValidLowLevelFast())
+				{
+					if (const FProperty* CachedProperty = Item->CachedProperty)
+					{
+						const UAudioPropertiesSheetAsset* SourceAsPropertySheet = Cast<UAudioPropertiesSheetAsset>(SourceObjectPtr);
+						const void* PropertyValuePtr = nullptr;
+						
+						if (SourceAsPropertySheet)
+						{
+							auto OnLeafPropertyVisited = [&CachedProperty, &PropertyValuePtr](const FPropertyBagPropertyDesc& PropertyDesc, const FAudioPropertiesSheet& OwningSheet)
+							{
+								if (PropertyDesc.Name == CachedProperty->GetName())
+								{
+									PropertyValuePtr = CachedProperty->ContainerPtrToValuePtr<void>(OwningSheet.Properties.GetValue().GetMemory());
+								}
+							};
+								
+							AudioPropertiesUtils::VisitLeafMostProperties(*SourceAsPropertySheet, OnLeafPropertyVisited);
+								
+						}
+						else
+						{
+							PropertyValuePtr = CachedProperty->ContainerPtrToValuePtr<void>(SourceObjectPtr);
+						}
+						
+						if(PropertyValuePtr)
+						{
+							CachedProperty->ExportTextItem_Direct(
+								PropertyValue,
+								PropertyValuePtr,   
+								nullptr,
+								nullptr,
+								PPF_None			
+							);		
+						}
+					}
+				}
+				
+				return SNew(SHorizontalBox)
+					+ SHorizontalBox::Slot()
+					.AutoWidth()
+					.Padding(4.f, ColumnItemBoxPadding)
+					[
+						SNew(STextBlock)
+						.AutoWrapText(true)
+						.Text(FText::FromString(PropertyValue))
+						.ToolTipText(FText::FromString(PropertyValue))
+					];
+			}
+
+			return SNullWidget::NullWidget;
+		}
+
+	private:
+		AudioPropertiesSheetAssetBuilder::FPropertyRequestPtr Item;
+		const UObject* SourceObjectPtr = nullptr;
+	};
+}
+
+void SAudioPropertiesSheetBuilderWidget::Construct(const FArguments& InArgs)
+{
+    SetSourceObject(InArgs._SourceObject);
+
+	if (InArgs._bSourceIsParent && SourceObject)
+	{
+		if (ensureAlwaysMsgf(bSourceIsPropertySheet, TEXT("Source should be property sheet to be parent")))
+		{
+			SetParentAsset(GetSourceAssetData());
+		}
+	}
+
+    ChildSlot
+    [
+		SNew(SVerticalBox)
+		+ SVerticalBox::Slot()
+		.AutoHeight()
+		[
+			GenerateSourceObjectWidget()
+		]
+    	+ SVerticalBox::Slot()
+		.AutoHeight()
+		[
+			GenerateParentPickerWidget()
+		]
+		+ SVerticalBox::Slot()
+		[
+			GeneratePropertyListWidget()
+		]
+		+ SVerticalBox::Slot()
+        .Padding(10.0f)
+		.AutoHeight()
+        [
+            SNew(SButton)
+            .Text(FText::FromString("Make Property Sheet"))
+			.HAlign(EHorizontalAlignment::HAlign_Center)
+            .OnClicked(this, &SAudioPropertiesSheetBuilderWidget::OnMakePropertySheetClicked)
+        ]
+    ];
+}
+
+TSharedRef<SWidget> SAudioPropertiesSheetBuilderWidget::GenerateSourceObjectWidget()
+{
+	TSharedRef<SWidget> ComboBoxContent = SNew(SHorizontalBox)
+	+ SHorizontalBox::Slot()
+	.FillWidth(1.0f)
+	.VAlign(VAlign_Center)
+	.Padding(5, 0)
+	[
+		SAssignNew(SourcePickerComboText, STextBlock)
+		.Text(FText::FromString(SourceObject.GetName()))
+	];
+
+	AudioPropertiesSheetBuilderWidget::FAssetPickerParams AssetPickerParams;
+	AssetPickerParams.InitialObject = GeneratedSheetParent;
+	AssetPickerParams.AllowedClasses = {UAudioPropertiesSheetAsset::StaticClass(), USoundBase::StaticClass()};
+	AssetPickerParams.OnAssetSelected = FOnAssetSelected::CreateSP(this, &SAudioPropertiesSheetBuilderWidget::SetSourceAsset);
+
+	return SNew(SVerticalBox)
+		+ SVerticalBox::Slot()
+		.AutoHeight()
+		.Padding(AudioPropertiesSheetBuilderWidget::HeaderItemPadding)
+		[
+			SNew(SHorizontalBox)
+			+ SHorizontalBox::Slot()
+			.AutoWidth()
+			.VAlign(VAlign_Center)
+			[
+				SNew(STextBlock)
+				.Text(FText::FromString("Source Object"))
+				.Font(FCoreStyle::Get().GetFontStyle("BoldFont"))
+				.ToolTipText(FText::FromString("The object we are inheriting properties from"))
+			]
+			+ SHorizontalBox::Slot()
+			.VAlign(VAlign_Center)
+			.Padding(AudioPropertiesSheetBuilderWidget::ComboBoxPadding)
+			.AutoWidth()
+			[
+				SAssignNew(SourcePickerComboButton, SComboButton)
+				.OnGetMenuContent_Lambda([AssetPickerParams](){return AudioPropertiesSheetBuilderWidget::BuildAssetPickerWidget(AssetPickerParams);})
+				.ButtonContent()
+				[
+					ComboBoxContent
+				]
+			]
+			+ SHorizontalBox::Slot()
+			.AutoWidth()
+			.Padding(2.0f, 0.f)
+			.VAlign(VAlign_Center)
+			[
+				PropertyCustomizationHelpers::MakeUseSelectedButton(
+					FSimpleDelegate::CreateSP(this, &SAudioPropertiesSheetBuilderWidget::OnSourceUseSelectedButtonClicked)
+				)
+			]
+			+ SHorizontalBox::Slot()
+			.AutoWidth()
+			.Padding(2.0f, 0.f)
+			.VAlign(VAlign_Center)
+			[
+				PropertyCustomizationHelpers::MakeBrowseButton(
+					FSimpleDelegate::CreateSP(this, &SAudioPropertiesSheetBuilderWidget::OnSourceBrowseToButtonClicked)
+				)
+			]
+			+ SHorizontalBox::Slot()
+			.AutoWidth()
+			.Padding(0.f, 0.f)
+			.VAlign(VAlign_Center)
+			[
+				PropertyCustomizationHelpers::MakeEditButton(
+					FSimpleDelegate::CreateSP(this, &SAudioPropertiesSheetBuilderWidget::OnSourceEditButtonClicked)
+				)
+			]
+		];
+}
+
+TSharedRef<SWidget> SAudioPropertiesSheetBuilderWidget::GenerateParentPickerWidget()
+{
+	TSharedRef<SWidget> ComboBoxContent = SNew(SHorizontalBox)
+	+ SHorizontalBox::Slot()
+	.FillWidth(1.0f)
+	.VAlign(VAlign_Center)
+	.Padding(5, 0)
+	[
+		SAssignNew(ParentPickerComboText, STextBlock)
+		.Text(FText::FromString(GeneratedSheetParent.IsValid() ? GeneratedSheetParent.AssetName.ToString() : "None"))
+	];
+
+	AudioPropertiesSheetBuilderWidget::FAssetPickerParams AssetPickerParams;
+	AssetPickerParams.InitialObject = GeneratedSheetParent;
+	AssetPickerParams.bAllowClear = true;
+	AssetPickerParams.AllowedClasses = {UAudioPropertiesSheetAsset::StaticClass()};
+	AssetPickerParams.OnAssetSelected = FOnAssetSelected::CreateSP(this, &SAudioPropertiesSheetBuilderWidget::SetParentAsset);
+
+	return SNew(SVerticalBox)
+		+ SVerticalBox::Slot()
+		.AutoHeight()
+		.Padding(AudioPropertiesSheetBuilderWidget::HeaderItemPadding)
+		[
+			SNew(SHorizontalBox)
+			+ SHorizontalBox::Slot()
+			.AutoWidth()
+			.VAlign(VAlign_Center)
+			[
+				SNew(STextBlock)
+				.Text(FText::FromString("Generated Sheet Parent"))
+				.ToolTipText(FText::FromString("The parent that will be assigned to the generated property sheet"))
+			]
+			+ SHorizontalBox::Slot()
+			.VAlign(VAlign_Center)
+			.AutoWidth()
+			.Padding(AudioPropertiesSheetBuilderWidget::ComboBoxPadding)
+			[
+				SAssignNew(ParentPickerComboButton, SComboButton)
+				.OnGetMenuContent_Lambda([AssetPickerParams](){return AudioPropertiesSheetBuilderWidget::BuildAssetPickerWidget(AssetPickerParams);})
+				.ButtonContent()
+				[
+					ComboBoxContent
+				]
+			]
+			+ SHorizontalBox::Slot()
+			.AutoWidth()
+			.Padding(2.0f, 0.f)
+			.VAlign(VAlign_Center)
+			[
+				PropertyCustomizationHelpers::MakeUseSelectedButton(
+					FSimpleDelegate::CreateSP(this, &SAudioPropertiesSheetBuilderWidget::OnParentUseSelectedButtonClicked)
+				)
+			]
+			+ SHorizontalBox::Slot()
+			.AutoWidth()
+			.Padding(2.0f, 0.f)
+			.VAlign(VAlign_Center)
+			[
+				PropertyCustomizationHelpers::MakeBrowseButton(
+					FSimpleDelegate::CreateSP(this, &SAudioPropertiesSheetBuilderWidget::OnParentBrowseToButtonClicked)
+				)
+			]
+			+ SHorizontalBox::Slot()
+			.AutoWidth()
+			.Padding(0.f, 0.f)
+			.VAlign(VAlign_Center)
+			[
+				PropertyCustomizationHelpers::MakeEditButton(
+					FSimpleDelegate::CreateSP(this, &SAudioPropertiesSheetBuilderWidget::OnParentEditButtonClicked)
+				)
+			]
+		];
+}
+
+void SAudioPropertiesSheetBuilderWidget::UpdatePropertyRequests()
+{
+	PropertyRequests.Empty();
+
+	UClass* SourceClass = SourceObject->GetClass();
+	
+	if (SourceClass && !bSourceIsPropertySheet)
+	{
+		for (TFieldIterator<FProperty> PropIt(SourceClass); PropIt; ++PropIt)
+		{
+			FProperty* Property = *PropIt;
+
+			if (!Property)
+			{
+				continue;
+			}
+
+			if (Property->HasAllPropertyFlags(CPF_Edit) && !Property->HasAnyPropertyFlags(CPF_EditConst))
+			{
+				PropertyRequests.Add(MakeShared<AudioPropertiesSheetAssetBuilder::FPropertyRequest>(Property));
+			}
+		}
+	}
+	else if (SourceObject && bSourceIsPropertySheet)
+	{
+		const UAudioPropertiesSheetAsset* SourcePropertySheet = Cast<UAudioPropertiesSheetAsset>(SourceObject.Get());
+
+		if (SourcePropertySheet)
+		{
+			auto OnLeafPropertyVisited = [this](const FPropertyBagPropertyDesc& PropertyDesc, const FAudioPropertiesSheet& OwningSheet)
+				{
+					PropertyRequests.Add(MakeShared<AudioPropertiesSheetAssetBuilder::FPropertyRequest>(PropertyDesc.CachedProperty));
+				};
+			
+			AudioPropertiesUtils::VisitLeafMostProperties(*SourcePropertySheet, OnLeafPropertyVisited);
+		}
+		
+	}
+
+	FilteredPropertyRequests = PropertyRequests;
+
+	if (PropertyRequestsView.IsValid())
+	{
+		PropertyRequestsView->RequestListRefresh();
+	}
+}
+
+TSharedRef<SWidget> SAudioPropertiesSheetBuilderWidget::GeneratePropertyListWidget()
+{
+	PropertyRequestsView = SNew(SListView<AudioPropertiesSheetAssetBuilder::FPropertyRequestPtr>)
+		.ListItemsSource(&FilteredPropertyRequests)
+		.OnGenerateRow(this, &SAudioPropertiesSheetBuilderWidget::OnGenerateRowForList)
+		.HeaderRow
+		(
+			SNew(SHeaderRow)
+			+ SHeaderRow::Column(AudioPropertiesSheetBuilderWidget::AddPropertyColumnName)
+				.DefaultLabel(LOCTEXT("AddPropertyColumnHeaderName", "Add Property"))
+				.HAlignHeader(HAlign_Center)
+				.HAlignCell(HAlign_Center)
+				.ManualWidth(100)
+				.SortMode(this, &SAudioPropertiesSheetBuilderWidget::GetColumnSortMode, AudioPropertiesSheetBuilderWidget::AddPropertyColumnName)
+				.OnSort(this, &SAudioPropertiesSheetBuilderWidget::OnColumnSortModeChanged)
+				.InitialSortMode(EColumnSortMode::Type::Ascending)
+			+ SHeaderRow::Column(AudioPropertiesSheetBuilderWidget::CopyPropertyValueColumnName)
+				.DefaultLabel(LOCTEXT("CopyPropertyValueColumnHeaderName", "Copy Property Value"))
+				.HAlignHeader(HAlign_Center)
+				.HAlignCell(HAlign_Center)
+				.ManualWidth(200)
+				.SortMode(this, &SAudioPropertiesSheetBuilderWidget::GetColumnSortMode, AudioPropertiesSheetBuilderWidget::CopyPropertyValueColumnName)
+				.OnSort(this, &SAudioPropertiesSheetBuilderWidget::OnColumnSortModeChanged)
+			+ SHeaderRow::Column(AudioPropertiesSheetBuilderWidget::PropertyNameColumnName)
+				.DefaultLabel(LOCTEXT("PropertyNameColumnHeaderName", "Property Name"))
+				.SortMode(this, &SAudioPropertiesSheetBuilderWidget::GetColumnSortMode, AudioPropertiesSheetBuilderWidget::PropertyNameColumnName)
+				.OnSort(this, &SAudioPropertiesSheetBuilderWidget::OnColumnSortModeChanged)
+			+ SHeaderRow::Column(AudioPropertiesSheetBuilderWidget::PropertyValueColumnName)
+				.DefaultLabel(LOCTEXT("PropertyValueColumnHeaderName", "Property Value"))
+		);
+
+    return SNew(SVerticalBox)
+		    + SVerticalBox::Slot()
+            .AutoHeight()
+			.Padding(AudioPropertiesSheetBuilderWidget::HeaderItemPadding)
+            [
+                SNew(SEditableTextBox)
+                .OnTextChanged(this, &SAudioPropertiesSheetBuilderWidget::OnSearchTextChanged)
+				.HintText(FText::FromString("Filter Properties"))
+            ]
+			+ SVerticalBox::Slot()
+			.AutoHeight()
+			.Padding(AudioPropertiesSheetBuilderWidget::HeaderItemPadding)
+			+ SVerticalBox::Slot()
+			[
+				PropertyRequestsView.ToSharedRef()
+			];
+}
+
+void SAudioPropertiesSheetBuilderWidget::UpdateParentPickerComboText() const
+{
+	if (ParentPickerComboText)
+	{
+		ParentPickerComboText->SetText(FText::FromString(GeneratedSheetParent.IsValid() ? GeneratedSheetParent.AssetName.ToString() : "None" ));
+	}
+}
+
+void SAudioPropertiesSheetBuilderWidget::UpdateSourcePickerComboText() const
+{
+	if (SourcePickerComboText)
+	{
+		SourcePickerComboText->SetText(FText::FromString(SourceObject ? SourceObject.GetName() : "None" ));
+	}
+}
+
+void SAudioPropertiesSheetBuilderWidget::SetSourceObject(TObjectPtr<const UObject> NewSourceObject)
+{
+	if (NewSourceObject == SourceObject)
+	{
+		return;
+	}
+	
+	UClass* OldSourceClass = SourceObject ? SourceObject->GetClass() : nullptr;
+	UClass* NewSourceClass = NewSourceObject->GetClass();
+	
+	bSourceIsPropertySheet = NewSourceClass ? NewSourceClass->IsChildOf<UAudioPropertiesSheetAsset>() : false;
+
+	if (!bSourceIsPropertySheet && OldSourceClass == NewSourceClass)
+	{
+		SourceObject = NewSourceObject;
+		UpdateSourcePickerComboText();
+		if (PropertyRequestsView.IsValid())
+		{
+			PropertyRequestsView->RebuildList();
+		}
+		
+		//needed to force property columns to show the new values
+		FSlateApplication::Get().GetRenderer()->FlushCommands();
+		
+		return;
+	}
+
+	const bool bShowDialog = OldSourceClass != nullptr;
+
+	if (bShowDialog)
+	{
+		FText Title = FText::FromString("Changing Source Object Class");
+		FText Message = FText::FromString("Changing Source Object Class will clear your property selection. Continue?");
+
+		EAppReturnType::Type DialogResponse = FMessageDialog::Open(EAppMsgType::YesNo, Message, Title);
+
+		if (DialogResponse == EAppReturnType::No)
+		{
+			return;
+		}
+	}
+	
+	SourceObject = NewSourceObject;
+	UpdateSourcePickerComboText();
+	UpdatePropertyRequests();
+}
+
+void SAudioPropertiesSheetBuilderWidget::SetSourceAsset(const FAssetData& InSelectedSource)
+{
+	UObject* NewSourceObject = InSelectedSource.GetAsset();
+	SetSourceObject(NewSourceObject);
+
+	if (SourcePickerComboButton)
+	{
+		SourcePickerComboButton->SetIsOpen(false);
+	}
+}
+
+void SAudioPropertiesSheetBuilderWidget::SetParentAsset(const FAssetData& InSelectedParent)
+{
+	UClass* AssetClass = InSelectedParent.GetClass();
+	
+	if (!ensureAlwaysMsgf(AssetClass && AssetClass->IsChildOf(UAudioPropertiesSheetAsset::StaticClass()), TEXT("Asset class needs to be a child of UAudioPropertiesSheetAsset to be set as parent")))
+	{
+		return;
+	}
+	
+	GeneratedSheetParent = InSelectedParent;
+	UpdateParentPickerComboText();
+
+	if (ParentPickerComboButton)
+	{
+		ParentPickerComboButton->SetIsOpen(false);
+	}
+}
+
+TSharedRef<ITableRow> SAudioPropertiesSheetBuilderWidget::OnGenerateRowForList(
+    AudioPropertiesSheetAssetBuilder::FPropertyRequestPtr InItem,
+    const TSharedRef<STableViewBase>& OwnerTable)
+{
+	return SNew(AudioPropertiesSheetBuilderWidget::SPropertyRow, OwnerTable, InItem, SourceObject.Get());
+}
+
+FReply SAudioPropertiesSheetBuilderWidget::OnMakePropertySheetClicked()
+{
+	TObjectPtr<UAudioPropertiesSheetAsset> ParentPtr = nullptr;
+
+	if (GeneratedSheetParent.IsValid())
+	{
+		ParentPtr = Cast<UAudioPropertiesSheetAsset>(GeneratedSheetParent.GetAsset());
+	}
+
+	AudioPropertiesSheetAssetBuilder::ESourceObjectParsingType ParsingType = bSourceIsPropertySheet ? AudioPropertiesSheetAssetBuilder::ESourceObjectParsingType::PropertySheet : AudioPropertiesSheetAssetBuilder::ESourceObjectParsingType::UObject;
+	
+	FAudioPropertiesSheetAssetBuilder::BuildPropertySheetFromPropertyDataArray(SourceObject, ParentPtr, PropertyRequests, ParsingType);
+		
+	return FReply::Handled();
+}
+
+void SAudioPropertiesSheetBuilderWidget::OnSearchTextChanged(const FText& SearchText)
+{
+	const FString SearchString = SearchText.ToString();
+	FilteredPropertyRequests.Empty();
+
+	for (const AudioPropertiesSheetAssetBuilder::FPropertyRequestPtr& PropertyRequest : PropertyRequests)
+	{
+		const FProperty* Property = PropertyRequest->CachedProperty;
+
+		if (!Property)
+		{
+			continue;
+		}
+
+		if (Property->GetFName().ToString().Contains(SearchString) || Property->GetDisplayNameText().ToString().Contains(SearchString))
+		{
+			FilteredPropertyRequests.Add(PropertyRequest);
+		}
+	}
+	
+	SortProperties();
+}
+
+void SAudioPropertiesSheetBuilderWidget::OnSourceBrowseToButtonClicked() const
+{
+	FSoftObjectPath AssetPath = SourceObject.GetPath();
+
+	if (AssetPath.IsValid())
+	{
+		// Find asset data to avoid loading the object
+		FAssetRegistryModule& AssetRegistryModule = FModuleManager::LoadModuleChecked<FAssetRegistryModule>("AssetRegistry");
+		FAssetData AssetData = AssetRegistryModule.Get().GetAssetByObjectPath(AssetPath);
+		GEditor->SyncBrowserToObject(AssetData);
+	}
+}
+
+void SAudioPropertiesSheetBuilderWidget::OnSourceEditButtonClicked() const
+{
+	GEditor->EditObject(const_cast<UObject*>(SourceObject.Get()));
+}
+
+void SAudioPropertiesSheetBuilderWidget::OnSourceUseSelectedButtonClicked()
+{
+	TArray<FAssetData> SelectedAssets;
+	GEditor->GetContentBrowserSelections(SelectedAssets);
+
+	if (SelectedAssets.IsEmpty())
+	{
+		return;
+	}
+
+	SetSourceAsset(SelectedAssets[0]);
+}
+
+void SAudioPropertiesSheetBuilderWidget::OnParentBrowseToButtonClicked() const
+{
+	GEditor->SyncBrowserToObject(GeneratedSheetParent);
+}
+
+void SAudioPropertiesSheetBuilderWidget::OnParentEditButtonClicked() const
+{
+	GEditor->EditObject(GeneratedSheetParent.GetAsset());
+}
+
+void SAudioPropertiesSheetBuilderWidget::OnParentUseSelectedButtonClicked()
+{
+	TArray<FAssetData> SelectedAssets;
+	GEditor->GetContentBrowserSelections(SelectedAssets);
+
+	if (SelectedAssets.IsEmpty())
+	{
+		return;
+	}
+
+	SetSourceAsset(SelectedAssets[0]);
+}
+
+EColumnSortMode::Type SAudioPropertiesSheetBuilderWidget::GetColumnSortMode(const FName InColumnName) const
+{
+	if ( SortedColumn == InColumnName )
+	{
+		return SortMode;
+	}
+
+	return EColumnSortMode::None;
+}
+
+void SAudioPropertiesSheetBuilderWidget::OnColumnSortModeChanged(const EColumnSortPriority::Type SortPriority, const FName& ColumnName, const EColumnSortMode::Type InSortMode)
+{
+	SortedColumn = ColumnName;
+	SortMode = InSortMode;
+	SortProperties();
+}
+
+void SAudioPropertiesSheetBuilderWidget::SortProperties()
+{
+	if ( SortMode != EColumnSortMode::None )
+	{
+		const bool bAscending = EColumnSortMode::Ascending ==  SortMode;
+		if ( SortedColumn == AudioPropertiesSheetBuilderWidget::PropertyNameColumnName )
+		{
+			FilteredPropertyRequests.StableSort( [bAscending](const AudioPropertiesSheetAssetBuilder::FPropertyRequestPtr& First, const AudioPropertiesSheetAssetBuilder::FPropertyRequestPtr& Second)
+					{
+						if(First->CachedProperty && Second->CachedProperty)
+						{
+							return First->CachedProperty->GetName() < Second->CachedProperty->GetName() == bAscending;
+						}
+						return false; 
+					}
+				);
+		}
+		else if (SortedColumn == AudioPropertiesSheetBuilderWidget::AddPropertyColumnName )
+		{
+			FilteredPropertyRequests.StableSort( [bAscending](const AudioPropertiesSheetAssetBuilder::FPropertyRequestPtr& First, const AudioPropertiesSheetAssetBuilder::FPropertyRequestPtr& Second)
+				{
+					return ((First->bInheritValue + First->bInheritProperty) > (Second->bInheritValue + Second->bInheritProperty)) == bAscending;
+				}
+			);
+		}
+		else if (SortedColumn == AudioPropertiesSheetBuilderWidget::CopyPropertyValueColumnName )
+		{
+			FilteredPropertyRequests.StableSort( [bAscending](const AudioPropertiesSheetAssetBuilder::FPropertyRequestPtr& First, const AudioPropertiesSheetAssetBuilder::FPropertyRequestPtr& Second)
+				{
+					return ((First->bInheritValue + First->bInheritProperty) > (Second->bInheritValue + Second->bInheritProperty)) == bAscending;
+				}
+			);
+		}
+		else
+		{
+			ensureAlwaysMsgf(false, TEXT("Trying to sort an unsupported column"));
+		}
+	}
+
+	if (PropertyRequestsView.IsValid())
+	{
+		PropertyRequestsView->RequestListRefresh();
+	}
+	
+}
+
+FAssetData SAudioPropertiesSheetBuilderWidget::GetSourceAssetData()
+{
+	if (ensure(SourceObject))
+	{
+		const FString ObjectPath = SourceObject->GetPathName();
+		FAssetRegistryModule& AssetRegistryModule = FModuleManager::LoadModuleChecked<FAssetRegistryModule>("AssetRegistry");
+		return AssetRegistryModule.Get().GetAssetByObjectPath(ObjectPath);
+	}
+
+	return FAssetData();
+}
+
+#undef LOCTEXT_NAMESPACE
\ No newline at end of file
diff --git a/Engine/Plugins/Experimental/AudioProperties/Source/AudioPropertiesEditor/Private/SAudioPropertiesSheetBuilderWidget.h b/Engine/Plugins/Experimental/AudioProperties/Source/AudioPropertiesEditor/Private/SAudioPropertiesSheetBuilderWidget.h
new file mode 100644
--- /dev/null
+++ b/Engine/Plugins/Experimental/AudioProperties/Source/AudioPropertiesEditor/Private/SAudioPropertiesSheetBuilderWidget.h
@@ -0,0 +1,82 @@
+// Copyright Epic Games, Inc. All Rights Reserved.
+
+#pragma once
+
+#include "AssetRegistry/AssetData.h"
+#include "AudioPropertiesSheetAssetBuilder.h"
+#include "UObject/SoftObjectPtr.h"
+#include "Widgets/Views/SHeaderRow.h"
+#include "Widgets/Views/SListView.h"
+#include "Widgets/SCompoundWidget.h"
+
+class ITableRow;
+class SComboButton;
+class STextBlock;
+class UObject;
+class UClass;
+
+class SAudioPropertiesSheetBuilderWidget : public SCompoundWidget
+{
+public:
+    SLATE_BEGIN_ARGS(SAudioPropertiesSheetBuilderWidget) {}
+	    
+	    //The object to inherit properties from 
+		SLATE_ARGUMENT(TObjectPtr<const UObject>, SourceObject)
+
+	    //If true, sets the source object as the parent on widget construction
+    	SLATE_ARGUMENT(bool, bSourceIsParent)
+    SLATE_END_ARGS()
+
+    void Construct(const FArguments& InArgs);
+
+private:
+	TSharedRef<SWidget> GenerateSourceObjectWidget();
+    TSharedRef<SWidget> GeneratePropertyListWidget();
+	TSharedRef<SWidget> GenerateParentPickerWidget();
+	void UpdateSourcePickerComboText() const;
+	void UpdateParentPickerComboText() const;
+	
+	void SetSourceObject(TObjectPtr<const UObject> NewSourceObject);
+	void SetSourceAsset(const FAssetData& InSelectedSource);
+	void SetParentAsset(const FAssetData& InSelectedParent);
+	void UpdatePropertyRequests();
+
+	FAssetData GetSourceAssetData();
+
+	//Widgets callbacks
+	TSharedRef<ITableRow> OnGenerateRowForList(AudioPropertiesSheetAssetBuilder::FPropertyRequestPtr InItem, const TSharedRef<STableViewBase>& OwnerTable);
+	FReply OnMakePropertySheetClicked();
+	void OnSearchTextChanged(const FText& SearchText);
+
+	//Asset browsing delegates - source
+	void OnSourceBrowseToButtonClicked() const;
+	void OnSourceEditButtonClicked() const;
+	void OnSourceUseSelectedButtonClicked();
+
+	//Asset browsing delegates - parent
+	void OnParentBrowseToButtonClicked() const;
+	void OnParentEditButtonClicked() const;
+	void OnParentUseSelectedButtonClicked();
+
+	//Sorting methods
+    EColumnSortMode::Type GetColumnSortMode(FName InColumnName) const;
+    void OnColumnSortModeChanged(EColumnSortPriority::Type SortPriority, const FName& ColumnName, EColumnSortMode::Type InSortMode);
+    void SortProperties();
+
+	TObjectPtr<const UObject> SourceObject;
+	FAssetData GeneratedSheetParent;
+	bool bSourceIsPropertySheet = false;
+	
+    TSharedPtr<SListView<AudioPropertiesSheetAssetBuilder::FPropertyRequestPtr>> PropertyRequestsView;
+	TArray<AudioPropertiesSheetAssetBuilder::FPropertyRequestPtr> PropertyRequests;
+	TArray<AudioPropertiesSheetAssetBuilder::FPropertyRequestPtr> FilteredPropertyRequests;
+	
+	FName SortedColumn;
+	EColumnSortMode::Type SortMode;
+	
+	TSharedPtr<SComboButton> SourcePickerComboButton;
+	TSharedPtr<STextBlock> SourcePickerComboText;
+
+	TSharedPtr<SComboButton> ParentPickerComboButton;
+	TSharedPtr<STextBlock> ParentPickerComboText;
+};
\ No newline at end of file
diff --git a/Engine/Plugins/Experimental/AudioProperties/Source/AudioPropertiesEditor/Public/AssetDefinition_AudioPropertiesSheetAsset.h b/Engine/Plugins/Experimental/AudioProperties/Source/AudioPropertiesEditor/Public/AssetDefinition_AudioPropertiesSheetAsset.h
new file mode 100644
--- /dev/null
+++ b/Engine/Plugins/Experimental/AudioProperties/Source/AudioPropertiesEditor/Public/AssetDefinition_AudioPropertiesSheetAsset.h
@@ -0,0 +1,32 @@
+﻿// Copyright Epic Games, Inc. All Rights Reserved.
+
+#pragma once
+
+#include "AssetTypeActions/AssetDefinitionDefault_AudioDiffable.h"
+#include "AudioPropertiesEditorModule.h"
+#include "AudioPropertiesSheet.h"
+
+#include "AssetDefinition_AudioPropertiesSheetAsset.generated.h"
+
+struct FToolMenuSection;
+
+UCLASS(MinimalAPI)
+class UAssetDefinition_AudioPropertiesSheetAsset : public UAssetDefinitionDefault_AudioDiffable
+{
+	GENERATED_BODY()
+
+public:
+	//~ Begin UAssetDefinition
+	virtual FText GetAssetDisplayName() const override { return NSLOCTEXT("AssetTypeActions", "AssetTypeActions_AudioPropertiesSheetAsset", "Audio Properties Sheet"); }
+	virtual FLinearColor GetAssetColor() const override { return FLinearColor(AudioPropertiesEditorModule::AssetColor); }
+	virtual TSoftClassPtr<UObject> GetAssetClass() const override { return UAudioPropertiesSheetAsset::StaticClass(); }
+	virtual TConstArrayView<FAssetCategoryPath> GetAssetCategories() const override;
+	virtual EAssetCommandResult OpenAssets(const FAssetOpenArgs& OpenArgs) const override;
+	//~ End UAssetDefinition
+
+	FToolMenuSection* FindContextMenuSection(FName SectionName) const;
+	
+private:
+	virtual TArray<FToolMenuSection*> RebuildContextMenuSections() const;
+	
+};
diff --git a/Engine/Plugins/Experimental/AudioProperties/Source/AudioPropertiesEditor/Public/AudioPropertiesBindingsFactory.h b/Engine/Plugins/Experimental/AudioProperties/Source/AudioPropertiesEditor/Public/AudioPropertiesBindingsFactory.h
new file mode 100644
--- /dev/null
+++ b/Engine/Plugins/Experimental/AudioProperties/Source/AudioPropertiesEditor/Public/AudioPropertiesBindingsFactory.h
@@ -0,0 +1,16 @@
+// Copyright Epic Games, Inc. All Rights Reserved.
+
+#pragma once
+
+#include "UObject/ObjectMacros.h"
+#include "Factories/Factory.h"
+
+#include "AudioPropertiesBindingsFactory.generated.h"
+
+UCLASS(hidecategories = Object, MinimalAPI)
+class UAudioPropertiesBindingsFactory : public UFactory
+{
+	GENERATED_UCLASS_BODY()
+
+	virtual UObject* FactoryCreateNew(UClass* Class, UObject* InParent, FName Name, EObjectFlags Flags, UObject* Context, FFeedbackContext* Warn) override;
+};
diff --git a/Engine/Plugins/Experimental/AudioProperties/Source/AudioPropertiesEditor/Public/AudioPropertiesDetailsInjector.h b/Engine/Plugins/Experimental/AudioProperties/Source/AudioPropertiesEditor/Public/AudioPropertiesDetailsInjector.h
new file mode 100644
--- /dev/null
+++ b/Engine/Plugins/Experimental/AudioProperties/Source/AudioPropertiesEditor/Public/AudioPropertiesDetailsInjector.h
@@ -0,0 +1,54 @@
+// Copyright Epic Games, Inc. All Rights Reserved.
+
+#pragma once
+
+#include "IAudioPropertiesDetailsInjector.h"
+#include "Delegates/IDelegateInstance.h"
+#include "UObject/ObjectKey.h"
+
+#define UE_API AUDIOPROPERTIESEDITOR_API
+
+class IDetailLayoutBuilder;
+class IPropertyHandle;
+
+/**
+ * Helper to customize the IDetailCustomization of a class receiving properties from an AudioPropertySheetAsset
+ */
+class FAudioPropertiesDetailsInjector : public IAudioPropertiesDetailsInjector
+{
+public:
+	/** To be called from IDetailCustomization::CustomizeDetails to customize widgets of injected properties 
+	*	@param DetailBuilder					DetailBuilder from CustomizeDetails.
+	*	@param PropertySheetPropertyHandle		Property Handle of the property sheet injecting properties on the target UObject. 
+	*/	
+	UE_API void CustomizeInjectedPropertiesDetails(IDetailLayoutBuilder& DetailBuilder, TSharedRef<IPropertyHandle> PropertySheetPropertyHandle) override;
+
+	/** To be called from IDetailCustomization::CustomizeDetails to bind the detail view to changes on the property sheet
+	*	@param DetailBuilder					DetailBuilder from CustomizeDetails.
+	*	@param PropertySheetPropertyHandle		Property Handle of the property sheet injecting properties on the target UObject. 
+	*/
+	UE_API void BindDetailCustomizationToPropertySheetChanges(IDetailLayoutBuilder& DetailBuilder, TSharedRef<IPropertyHandle> PropertySheetPropertyHandle) override;
+
+	/** Call this on destruction/pending delete of the details customization to unbind the detail view to changes on the property sheet */
+	UE_API void UnbindFromPropertySheetChanges() override;
+
+private:
+	UE_API void BindDetailsRefreshToPropertySheetSwaps(const IDetailLayoutBuilder& DetailBuilder, TArray<TWeakObjectPtr<UObject>>& ObjectsBeingCustomized, TSharedRef<IPropertyHandle> PropertySheetPropertyHandle);
+	UE_API void BindDetailsRefreshToPropertySheetChanges(IDetailLayoutBuilder& DetailBuilder, TArray<TWeakObjectPtr<UObject>>& ObjectsBeingCustomized, TSharedRef<IPropertyHandle> PropertySheetPropertyHandle);
+
+	UE_API void CreateOverriddenPropertiesWidgets(IDetailLayoutBuilder& DetailBuilder, TArray<TWeakObjectPtr<UObject>> ObjectsBeingCustomized, TSharedRef<IPropertyHandle> PropertySheetPropertyHandle);
+
+	TArray<FDelegateHandle> OnTargetObjectPropertyChangeDelegateHandles;
+	TMultiMap<FObjectKey, FDelegateHandle> InjectedPropertiesBindings;
+};
+
+/**
+ * Modular feature builder to instantiate an FAudioPropertiesDetailsInjector
+ */
+class FAudioPropertiesDetailsInjectorBuilder : public IAudioPropertiesDetailsInjectorBuilder
+{
+public:
+	virtual IAudioPropertiesDetailsInjector* CreateAudioPropertiesDetailsInjector() override {return new FAudioPropertiesDetailsInjector; }
+};
+
+#undef UE_API
diff --git a/Engine/Plugins/Experimental/AudioProperties/Source/AudioPropertiesEditor/Public/AudioPropertiesEditorModule.h b/Engine/Plugins/Experimental/AudioProperties/Source/AudioPropertiesEditor/Public/AudioPropertiesEditorModule.h
new file mode 100644
--- /dev/null
+++ b/Engine/Plugins/Experimental/AudioProperties/Source/AudioPropertiesEditor/Public/AudioPropertiesEditorModule.h
@@ -0,0 +1,28 @@
+// Copyright Epic Games, Inc. All Rights Reserved.
+
+#pragma once
+
+#include "Math/Color.h"
+#include "Modules/ModuleInterface.h"
+#include "Templates/SharedPointer.h"
+
+class IAudioPropertiesDetailsInjector;
+class FAudioPropertiesSheetBuilderInstantiator;
+class FAudioPropertiesDetailsInjectorBuilder;
+
+namespace AudioPropertiesEditorModule
+{
+	inline constexpr FColor AssetColor = FColor(201, 29, 85);
+}
+
+class FAudioPropertiesEditorModule : public IModuleInterface
+{
+public:
+	virtual void StartupModule() override;
+	virtual void ShutdownModule() override;
+
+private:
+	TSharedPtr<FAudioPropertiesSheetBuilderInstantiator> PropertiesSheetBuilder = nullptr;
+	
+	TSharedPtr<FAudioPropertiesDetailsInjectorBuilder> DetailsInjectorBuilder = nullptr;
+};
\ No newline at end of file
diff --git a/Engine/Plugins/Experimental/AudioProperties/Source/AudioPropertiesEditor/Public/AudioPropertiesSheetAssetFactory.h b/Engine/Plugins/Experimental/AudioProperties/Source/AudioPropertiesEditor/Public/AudioPropertiesSheetAssetFactory.h
new file mode 100644
--- /dev/null
+++ b/Engine/Plugins/Experimental/AudioProperties/Source/AudioPropertiesEditor/Public/AudioPropertiesSheetAssetFactory.h
@@ -0,0 +1,16 @@
+// Copyright Epic Games, Inc. All Rights Reserved.
+
+#pragma once
+
+#include "UObject/ObjectMacros.h"
+#include "Factories/Factory.h"
+
+#include "AudioPropertiesSheetAssetFactory.generated.h"
+
+UCLASS(hidecategories = Object, MinimalAPI)
+class UAudioPropertiesSheetAssetFactory : public UFactory
+{
+	GENERATED_UCLASS_BODY()
+
+	virtual UObject* FactoryCreateNew(UClass* Class, UObject* InParent, FName Name, EObjectFlags Flags, UObject* Context, FFeedbackContext* Warn) override;
+};
diff --git a/Engine/Plugins/Experimental/AudioProperties/Source/AudioPropertiesEditor/Public/InstancedAudioPropertyBagDetails.h b/Engine/Plugins/Experimental/AudioProperties/Source/AudioPropertiesEditor/Public/InstancedAudioPropertyBagDetails.h
new file mode 100644
--- /dev/null
+++ b/Engine/Plugins/Experimental/AudioProperties/Source/AudioPropertiesEditor/Public/InstancedAudioPropertyBagDetails.h
@@ -0,0 +1,39 @@
+// Copyright Epic Games, Inc. All Rights Reserved.
+
+#pragma once
+
+#include "IDetailCustomization.h"
+#include "InstancedStructDetails.h"
+#include "PropertyBagDetails.h"
+
+struct FAudioPropertiesSheet;
+class UAudioPropertiesSheetAsset;
+
+
+class FInstancedAudioPropertyBagDetails : public FPropertyBagInstanceDataDetails
+{
+public:
+
+	FInstancedAudioPropertyBagDetails(
+		TSharedPtr<IPropertyHandle> InInstancedPropertyBagHandle,
+		const TSharedPtr<IPropertyUtilities>& InPropUtils,
+		const UAudioPropertiesSheetAsset* InOwningAudioPropertiesSheetAsset,
+		FAudioPropertiesSheet* InLeafMostPropertySheet, 
+		bool bInIsLeafBag = false);
+
+	virtual void OnChildRowAdded(IDetailPropertyRow& ChildRow)override;
+
+	virtual void GenerateHeaderRowContent(FDetailWidgetRow& NodeRow) override {}
+
+private:
+	void DrawLeafBagChildRow(IDetailPropertyRow& ChildRow);
+	void DrawInheritedBagChildRow(IDetailPropertyRow& ChildRow);
+
+	void HandlePropertyOverrideChanged(TSharedPtr<IPropertyHandle> PropertyThatChanged, const bool bOverrideState);
+	
+	TSharedPtr<IPropertyHandle> InstancedPropertyBagHandle = nullptr;
+	const UAudioPropertiesSheetAsset* OwningAudioPropertiesSheetAsset = nullptr;
+	FAudioPropertiesSheet* LeafMostPropertySheet = nullptr;
+	TSharedPtr<IPropertyUtilities> CachedPropUtils;
+	bool bIsLeafBag = false;
+};
\ No newline at end of file
