Walk UClass inheritance when overriding property details from a property sheet to avoid visualization errors

#Synced-CL 49735510

diff --git a/Engine/Plugins/Experimental/AudioProperties/Source/AudioPropertiesEditor/Private/AudioPropertiesDetailsInjector.cpp b/Engine/Plugins/Experimental/AudioProperties/Source/AudioPropertiesEditor/Private/AudioPropertiesDetailsInjector.cpp
--- a/Engine/Plugins/Experimental/AudioProperties/Source/AudioPropertiesEditor/Private/AudioPropertiesDetailsInjector.cpp
+++ b/Engine/Plugins/Experimental/AudioProperties/Source/AudioPropertiesEditor/Private/AudioPropertiesDetailsInjector.cpp
@@ -57,9 +57,17 @@ void FAudioPropertiesDetailsInjector::CreateOverriddenPropertiesWidgets(IDetailL
 
 			if (!OverriddenPropertyHandle->IsValidHandle())
 			{
-				//If GetProperty fails use the class path
+				//If GetProperty fails use the class path and walk the inheritance tree 
+				UClass* ObjectClass = ObjectBeingCustomized->GetClass();
+				for (UClass* Class = ObjectClass; Class != nullptr; Class = Class->GetSuperClass())
+				{
+					OverriddenPropertyHandle = DetailBuilder.GetProperty(TargetPropertyName, Class);
+					if (OverriddenPropertyHandle->IsValidHandle())
+					{
+						break;
+					}
+				}
 
-				OverriddenPropertyHandle = DetailBuilder.GetProperty(TargetPropertyName, ObjectBeingCustomized->GetClass());
 				if (!OverriddenPropertyHandle->IsValidHandle())
 				{
 					continue;
@@ -237,4 +245,4 @@ void FAudioPropertiesDetailsInjector::UnbindFromPropertySheetChanges()
 		}
 	}
 }
-#undef LOCTEXT_NAMESPACE 
\ No newline at end of file
+#undef LOCTEXT_NAMESPACE
\ No newline at end of file
