Instantiate Audio Properties Name Parser when creating new Property Sheet, as this is the de facto default
#Synced-CL 50441695
#Synced-CL 50441695

diff --git a/Engine/Plugins/Experimental/AudioProperties/Source/AudioProperties/Private/AudioPropertiesSheet.cpp b/Engine/Plugins/Experimental/AudioProperties/Source/AudioProperties/Private/AudioPropertiesSheet.cpp
--- a/Engine/Plugins/Experimental/AudioProperties/Source/AudioProperties/Private/AudioPropertiesSheet.cpp
+++ b/Engine/Plugins/Experimental/AudioProperties/Source/AudioProperties/Private/AudioPropertiesSheet.cpp
@@ -15,6 +15,8 @@
 #include "Misc/DataValidation.h"
 #endif
 
+#include "AudioPropertiesParserNameMatch.h"
+
 #include UE_INLINE_GENERATED_CPP_BY_NAME(AudioPropertiesSheet)
 
 #define LOCTEXT_NAMESPACE "AudioPropertiesSheetAsset"
@@ -256,6 +258,11 @@ void FAudioPropertiesSheet::FlattenSheet()
 UAudioPropertiesSheetAsset::UAudioPropertiesSheetAsset(const FObjectInitializer& ObjectInitializer)
 	: Super(ObjectInitializer)
 {
+	//if we don't have a parser (or are inheriting one) create a default one
+	if (!PropertiesParser && !PropertiesSheet.Parent)
+	{
+		PropertiesParser = ObjectInitializer.CreateDefaultSubobject<UAudioPropertiesParserNameMatch>(this, TEXT("DefaultPropertiesParser"));
+	}
 }
 
 #if WITH_EDITOR
